Removed pragma once in cpp file (fixed NonUnityTestCompile_OnMac, only clang detects this error)

diff --git a/Engine/Source/Developer/HierarchicalLODUtilities/Private/HierarchicalLODProxyProcessor.cpp b/Engine/Source/Developer/HierarchicalLODUtilities/Private/HierarchicalLODProxyProcessor.cpp
--- a/Engine/Source/Developer/HierarchicalLODUtilities/Private/HierarchicalLODProxyProcessor.cpp
+++ b/Engine/Source/Developer/HierarchicalLODUtilities/Private/HierarchicalLODProxyProcessor.cpp
@@ -1,7 +1,5 @@
 // Copyright 1998-2015 Epic Games, Inc. All Rights Reserved.
 
-#pragma once
-
 #include "HierarchicalLODUtilitiesModulePrivatePCH.h"
 #include "HierarchicalLODProxyProcessor.h"
 #include "HierarchicalLODUtilities.h"
