[NO NON-PII COMMIT MESSAGE]

diff --git a/Engine/Source/Developer/HierarchicalLODUtilities/Private/HierarchicalLODUtilities.cpp b/Engine/Source/Developer/HierarchicalLODUtilities/Private/HierarchicalLODUtilities.cpp
--- a/Engine/Source/Developer/HierarchicalLODUtilities/Private/HierarchicalLODUtilities.cpp
+++ b/Engine/Source/Developer/HierarchicalLODUtilities/Private/HierarchicalLODUtilities.cpp
@@ -783,34 +783,37 @@ bool FHierarchicalLODUtilities::IsWorldUsedForStreaming(const UWorld* InWorld)
 
 	for (const FAssetIdentifier& Identifier : ReferenceNames)
 	{
-		
-		const FString PackageName = Identifier.PackageName.ToString();
-		UPackage* ReferencingPackage = FindPackage(nullptr, *PackageName);
-		if (!ReferencingPackage)
-		{
-			ReferencingPackage = LoadPackage(nullptr, *PackageName, LOAD_None);
-		}
-
-		// Retrieve the referencing UPackage and check if it contains a map asset
-		if (ReferencingPackage && ReferencingPackage->ContainsMap())
+		// Referncers can include things like primary asset virtual packages, we don't want those
+		if (Identifier.PackageName != NAME_None)
 		{
-			TArray<UPackage*> Packages;
-			Packages.Add(ReferencingPackage);
-			TArray<UObject*> Objects;
-			PackageTools::GetObjectsInPackages(&Packages, Objects);
+			const FString PackageName = Identifier.PackageName.ToString();
+			UPackage* ReferencingPackage = FindPackage(nullptr, *PackageName);
+			if (!ReferencingPackage)
+			{
+				ReferencingPackage = LoadPackage(nullptr, *PackageName, LOAD_None);
+			}
 
-			// Loop over all objects in package and try to find a world
-			for (UObject* Object : Objects)
+			// Retrieve the referencing UPackage and check if it contains a map asset
+			if (ReferencingPackage && ReferencingPackage->ContainsMap())
 			{
-				if (UWorld* World = Cast<UWorld>(Object))
+				TArray<UPackage*> Packages;
+				Packages.Add(ReferencingPackage);
+				TArray<UObject*> Objects;
+				PackageTools::GetObjectsInPackages(&Packages, Objects);
+
+				// Loop over all objects in package and try to find a world
+				for (UObject* Object : Objects)
 				{
-					// Check the world contains InWorld as a streaming level
-					if (World->StreamingLevels.FindByPredicate([InWorld](const ULevelStreaming* StreamingLevel)
-					{
-						return StreamingLevel->GetWorldAsset() == InWorld;
-					}))
+					if (UWorld* World = Cast<UWorld>(Object))
 					{
-						return true;
+						// Check the world contains InWorld as a streaming level
+						if (World->StreamingLevels.FindByPredicate([InWorld](const ULevelStreaming* StreamingLevel)
+						{
+							return StreamingLevel->GetWorldAsset() == InWorld;
+						}))
+						{
+							return true;
+						}
 					}
 				}
 			}
