Crash building HLOD with imposters only

diff --git a/Engine/Source/Developer/HierarchicalLODUtilities/Private/HierarchicalLODUtilities.cpp b/Engine/Source/Developer/HierarchicalLODUtilities/Private/HierarchicalLODUtilities.cpp
--- a/Engine/Source/Developer/HierarchicalLODUtilities/Private/HierarchicalLODUtilities.cpp
+++ b/Engine/Source/Developer/HierarchicalLODUtilities/Private/HierarchicalLODUtilities.cpp
@@ -335,214 +335,220 @@ UStaticMesh* CreateImposterStaticMesh(UStaticMeshComponent* InComponent, UMateri
 
 bool FHierarchicalLODUtilities::BuildStaticMeshForLODActor(ALODActor* LODActor, UHLODProxy* Proxy, const FHierarchicalSimplification& LODSetup, UMaterialInterface* InBaseMaterial)
 {
-	if (Proxy && LODActor)
-	{
-		UE_LOG(LogHierarchicalLODUtilities, Log, TEXT("Building Proxy Mesh for Cluster %s"), *LODActor->GetName());
-		const FScopedTransaction Transaction(LOCTEXT("UndoAction_BuildProxyMesh", "Building Proxy Mesh for Cluster"));
+	if (!Proxy || !LODActor)
+	{	
+		return false;
+	}
+
+	UE_LOG(LogHierarchicalLODUtilities, Log, TEXT("Building Proxy Mesh for Cluster %s"), *LODActor->GetName());
+	const FScopedTransaction Transaction(LOCTEXT("UndoAction_BuildProxyMesh", "Building Proxy Mesh for Cluster"));
 
-		// Pass false here and dirty package later if values have changed
-		LODActor->Modify(false);
-		Proxy->Modify();
+	// Pass false here and dirty package later if values have changed
+	LODActor->Modify(false);
+	Proxy->Modify();
 
-		// Clean out the proxy as we are rebuilding meshes
-		Proxy->Clean();
-		UPackage* AssetsOuter = Proxy->GetOutermost();
+	// Clean out the proxy as we are rebuilding meshes
+	Proxy->Clean();
+	UPackage* AssetsOuter = Proxy->GetOutermost();
 
-		TArray<UPrimitiveComponent*> AllComponents;
-		UHLODProxy::ExtractComponents(LODActor, AllComponents);
+	TArray<UPrimitiveComponent*> AllComponents;
+	UHLODProxy::ExtractComponents(LODActor, AllComponents);
+
+	// It shouldn't even have come here if it didn't have any static meshes
+	if(!ensure(AllComponents.Num() > 0))
+	{
+		return false;
+	}
 
-		TArray<UStaticMeshComponent*> AllImposters;
-		if (LODSetup.MergeSetting.bIncludeImposters)
-		{			
-			// Retrieve all imposters.
-			for (UPrimitiveComponent* Component : AllComponents)
+	TArray<UStaticMeshComponent*> AllImposters;
+	if (LODSetup.MergeSetting.bIncludeImposters)
+	{			
+		// Retrieve all imposters.
+		for (UPrimitiveComponent* Component : AllComponents)
+		{
+			if (UStaticMeshComponent* StaticMeshComponent = Cast<UStaticMeshComponent>(Component))
 			{
-				if (UStaticMeshComponent* StaticMeshComponent = Cast<UStaticMeshComponent>(Component))
+				if (LODActor->GetImposterMaterial(StaticMeshComponent))
 				{
-					if (LODActor->GetImposterMaterial(StaticMeshComponent))
-					{
-						AllImposters.Add(StaticMeshComponent);
-					}
+					AllImposters.Add(StaticMeshComponent);
 				}
 			}
-
-			// Imposters won't be merged in the HLOD mesh
-			AllComponents.RemoveAll([&](UPrimitiveComponent* Component) { return AllImposters.Contains(Component); });
 		}
 
-		// it shouldn't even have come here if it didn't have any staticmesh
-		if (ensure(AllComponents.Num() > 0))
-		{
-			const AActor* FirstActor = UHLODProxy::FindFirstActor(LODActor);
+		// Imposters won't be merged in the HLOD mesh
+		AllComponents.RemoveAll([&](UPrimitiveComponent* Component) { return AllImposters.Contains(Component); });
+	}
+
+	if (AllComponents.Num() > 0)
+	{
+		const AActor* FirstActor = UHLODProxy::FindFirstActor(LODActor);
 
-			TArray<UObject*> OutAssets;
-			FVector OutProxyLocation = FVector::ZeroVector;
-			UStaticMesh* MainMesh = nullptr;
+		TArray<UObject*> OutAssets;
+		FVector OutProxyLocation = FVector::ZeroVector;
+		UStaticMesh* MainMesh = nullptr;
 
-			// Generate proxy mesh and proxy material assets
-			IMeshReductionManagerModule& MeshReductionModule = FModuleManager::Get().LoadModuleChecked<IMeshReductionManagerModule>("MeshReductionInterface");
-			const bool bHasMeshReductionCapableModule = (MeshReductionModule.GetMeshMergingInterface() != NULL);
+		// Generate proxy mesh and proxy material assets
+		IMeshReductionManagerModule& MeshReductionModule = FModuleManager::Get().LoadModuleChecked<IMeshReductionManagerModule>("MeshReductionInterface");
+		const bool bHasMeshReductionCapableModule = (MeshReductionModule.GetMeshMergingInterface() != NULL);
 
-			const IMeshMergeUtilities& MeshMergeUtilities = FModuleManager::Get().LoadModuleChecked<IMeshMergeModule>("MeshMergeUtilities").GetUtilities();
-			// should give unique name, so use level + actor name
+		const IMeshMergeUtilities& MeshMergeUtilities = FModuleManager::Get().LoadModuleChecked<IMeshMergeModule>("MeshMergeUtilities").GetUtilities();
+		// should give unique name, so use level + actor name
 			
-			const FString PackageName = FString::Printf(TEXT("LOD_%s_%i_%s"), *FirstActor->GetOutermost()->GetName(), LODActor->LODLevel - 1, *FirstActor->GetName());
-			if (bHasMeshReductionCapableModule && LODSetup.bSimplifyMesh)
-			{
-				FHierarchicalLODUtilitiesModule& Module = FModuleManager::LoadModuleChecked<FHierarchicalLODUtilitiesModule>("HierarchicalLODUtilities");
-				FHierarchicalLODProxyProcessor* Processor = Module.GetProxyProcessor();
+		const FString PackageName = FString::Printf(TEXT("LOD_%s_%i_%s"), *FirstActor->GetOutermost()->GetName(), LODActor->LODLevel - 1, *FirstActor->GetName());
+		if (bHasMeshReductionCapableModule && LODSetup.bSimplifyMesh)
+		{
+			FHierarchicalLODUtilitiesModule& Module = FModuleManager::LoadModuleChecked<FHierarchicalLODUtilitiesModule>("HierarchicalLODUtilities");
+			FHierarchicalLODProxyProcessor* Processor = Module.GetProxyProcessor();
 
-				FHierarchicalSimplification OverrideLODSetup = LODSetup;
+			FHierarchicalSimplification OverrideLODSetup = LODSetup;
 
-				FMeshProxySettings ProxySettings = LODSetup.ProxySetting;
-				if (LODActor->bOverrideMaterialMergeSettings)
-				{
-					ProxySettings.MaterialSettings = LODActor->MaterialSettings;
-				}
+			FMeshProxySettings ProxySettings = LODSetup.ProxySetting;
+			if (LODActor->bOverrideMaterialMergeSettings)
+			{
+				ProxySettings.MaterialSettings = LODActor->MaterialSettings;
+			}
 
-				if (LODActor->bOverrideScreenSize)
-				{
-					ProxySettings.ScreenSize = LODActor->ScreenSize;
-				}
+			if (LODActor->bOverrideScreenSize)
+			{
+				ProxySettings.ScreenSize = LODActor->ScreenSize;
+			}
 
-				if (LODActor->bOverrideTransitionScreenSize)
-				{
-					OverrideLODSetup.TransitionScreenSize = LODActor->TransitionScreenSize;
-				}
+			if (LODActor->bOverrideTransitionScreenSize)
+			{
+				OverrideLODSetup.TransitionScreenSize = LODActor->TransitionScreenSize;
+			}
 
-				FGuid JobID = Processor->AddProxyJob(LODActor, Proxy, OverrideLODSetup);
+			FGuid JobID = Processor->AddProxyJob(LODActor, Proxy, OverrideLODSetup);
 
-				TArray<UStaticMeshComponent*> StaticMeshComponents;
-				Algo::Transform(AllComponents, StaticMeshComponents, [](UPrimitiveComponent* InPrimitiveComponent) { return Cast<UStaticMeshComponent>(InPrimitiveComponent); });
+			TArray<UStaticMeshComponent*> StaticMeshComponents;
+			Algo::Transform(AllComponents, StaticMeshComponents, [](UPrimitiveComponent* InPrimitiveComponent) { return Cast<UStaticMeshComponent>(InPrimitiveComponent); });
 
-				MeshMergeUtilities.CreateProxyMesh(StaticMeshComponents, ProxySettings, InBaseMaterial, AssetsOuter, PackageName, JobID, Processor->GetCallbackDelegate(), true, OverrideLODSetup.TransitionScreenSize);
-			}
-			else
+			MeshMergeUtilities.CreateProxyMesh(StaticMeshComponents, ProxySettings, InBaseMaterial, AssetsOuter, PackageName, JobID, Processor->GetCallbackDelegate(), true, OverrideLODSetup.TransitionScreenSize);
+		}
+		else
+		{
+			FMeshMergingSettings MergeSettings = LODSetup.MergeSetting;
+			if (LODActor->bOverrideMaterialMergeSettings)
 			{
-				FMeshMergingSettings MergeSettings = LODSetup.MergeSetting;
-				if (LODActor->bOverrideMaterialMergeSettings)
-				{
-					MergeSettings.MaterialSettings = LODActor->MaterialSettings;
-				}
-
-				// update LOD parents before rebuild to ensure they are valid when mesh merge extensions are called.
-				LODActor->UpdateSubActorLODParents();
+				MergeSettings.MaterialSettings = LODActor->MaterialSettings;
+			}
 
-				MeshMergeUtilities.MergeComponentsToStaticMesh(AllComponents, FirstActor->GetWorld(), MergeSettings, InBaseMaterial, AssetsOuter, PackageName, OutAssets, OutProxyLocation, LODSetup.TransitionScreenSize, true);
+			// update LOD parents before rebuild to ensure they are valid when mesh merge extensions are called.
+			LODActor->UpdateSubActorLODParents();
 
-				// set staticmesh
-				for (UObject* Asset : OutAssets)
-				{
-					UStaticMesh* StaticMesh = Cast<UStaticMesh>(Asset);
+			MeshMergeUtilities.MergeComponentsToStaticMesh(AllComponents, FirstActor->GetWorld(), MergeSettings, InBaseMaterial, AssetsOuter, PackageName, OutAssets, OutProxyLocation, LODSetup.TransitionScreenSize, true);
 
-					if (StaticMesh)
-					{
-						MainMesh = StaticMesh;
-					}
-				}
+			// set staticmesh
+			for (UObject* Asset : OutAssets)
+			{
+				UStaticMesh* StaticMesh = Cast<UStaticMesh>(Asset);
 
-				if (!MainMesh)
+				if (StaticMesh)
 				{
-					return false;
+					MainMesh = StaticMesh;
 				}
+			}
 
-				// make sure the mesh won't affect navmesh generation
-				MainMesh->MarkAsNotHavingNavigationData();
+			if (!MainMesh)
+			{
+				return false;
+			}
 
-				bool bDirtyPackage = false;
-				UStaticMesh* PreviousStaticMesh = LODActor->GetStaticMeshComponent()->GetStaticMesh();
-				bDirtyPackage |= (MainMesh != PreviousStaticMesh);
-				LODActor->SetStaticMesh(MainMesh);
-				bDirtyPackage |= (LODActor->GetActorLocation() != OutProxyLocation);
-				LODActor->SetActorLocation(OutProxyLocation);
+			// make sure the mesh won't affect navmesh generation
+			MainMesh->MarkAsNotHavingNavigationData();
 
-				// Check resulting mesh and give a warning if it exceeds the vertex / triangle cap for certain platforms
-				FProjectStatus ProjectStatus;
-				if (IProjectManager::Get().QueryStatusForCurrentProject(ProjectStatus) && (ProjectStatus.IsTargetPlatformSupported(TEXT("Android")) || ProjectStatus.IsTargetPlatformSupported(TEXT("IOS"))))
-				{
-					if (MainMesh->RenderData.IsValid() && MainMesh->RenderData->LODResources.Num() && MainMesh->RenderData->LODResources[0].IndexBuffer.Is32Bit())
-					{
-						FMessageLog("HLODResults").Warning()
-							->AddToken(FUObjectToken::Create(LODActor))
-							->AddToken(FTextToken::Create(LOCTEXT("HLODError_MeshNotBuildTwo", " Mesh has more that 65535 vertices, incompatible with mobile; forcing 16-bit (will probably cause rendering issues).")));
-					}
-				}
-
-				// At the moment this assumes a fixed field of view of 90 degrees (horizontal and vertical axi)
-				static const float FOVRad = 90.0f * (float)PI / 360.0f;
-				static const FMatrix ProjectionMatrix = FPerspectiveMatrix(FOVRad, 1920, 1080, 0.01f);
-				FBoxSphereBounds Bounds = LODActor->GetStaticMeshComponent()->CalcBounds(FTransform());
+			bool bDirtyPackage = false;
+			UStaticMesh* PreviousStaticMesh = LODActor->GetStaticMeshComponent()->GetStaticMesh();
+			bDirtyPackage |= (MainMesh != PreviousStaticMesh);
+			LODActor->SetStaticMesh(MainMesh);
+			bDirtyPackage |= (LODActor->GetActorLocation() != OutProxyLocation);
+			LODActor->SetActorLocation(OutProxyLocation);
 
-				float DrawDistance;
-				if (LODSetup.bUseOverrideDrawDistance)
-				{
-					DrawDistance = LODSetup.OverrideDrawDistance;
-				}
-				else
+			// Check resulting mesh and give a warning if it exceeds the vertex / triangle cap for certain platforms
+			FProjectStatus ProjectStatus;
+			if (IProjectManager::Get().QueryStatusForCurrentProject(ProjectStatus) && (ProjectStatus.IsTargetPlatformSupported(TEXT("Android")) || ProjectStatus.IsTargetPlatformSupported(TEXT("IOS"))))
+			{
+				if (MainMesh->RenderData.IsValid() && MainMesh->RenderData->LODResources.Num() && MainMesh->RenderData->LODResources[0].IndexBuffer.Is32Bit())
 				{
-					DrawDistance = CalculateDrawDistanceFromScreenSize(Bounds.SphereRadius, LODSetup.TransitionScreenSize, ProjectionMatrix);
+					FMessageLog("HLODResults").Warning()
+						->AddToken(FUObjectToken::Create(LODActor))
+						->AddToken(FTextToken::Create(LOCTEXT("HLODError_MeshNotBuildTwo", " Mesh has more that 65535 vertices, incompatible with mobile; forcing 16-bit (will probably cause rendering issues).")));
 				}
+			}
 
-				bDirtyPackage |= (LODActor->GetDrawDistance() != DrawDistance);
-				LODActor->SetDrawDistance(DrawDistance);
-			
-				LODActor->DetermineShadowingFlags();
+			// At the moment this assumes a fixed field of view of 90 degrees (horizontal and vertical axi)
+			static const float FOVRad = 90.0f * (float)PI / 360.0f;
+			static const FMatrix ProjectionMatrix = FPerspectiveMatrix(FOVRad, 1920, 1080, 0.01f);
+			FBoxSphereBounds Bounds = LODActor->GetStaticMeshComponent()->CalcBounds(FTransform());
 
-				// Link proxy to actor
-				UHLODProxy* PreviousProxy = LODActor->GetProxy();
-				Proxy->AddMesh(LODActor, MainMesh, UHLODProxy::GenerateKeyForActor(LODActor));
-				bDirtyPackage |= (LODActor->GetProxy() != PreviousProxy);
+			float DrawDistance;
+			if (LODSetup.bUseOverrideDrawDistance)
+			{
+				DrawDistance = LODSetup.OverrideDrawDistance;
+			}
+			else
+			{
+				DrawDistance = CalculateDrawDistanceFromScreenSize(Bounds.SphereRadius, LODSetup.TransitionScreenSize, ProjectionMatrix);
+			}
 
-				if(bDirtyPackage)
-				{
-					LODActor->MarkPackageDirty();
-				}
+			bDirtyPackage |= (LODActor->GetDrawDistance() != DrawDistance);
+			LODActor->SetDrawDistance(DrawDistance);
+			
+			LODActor->DetermineShadowingFlags();
 
-				// Clean out standalone meshes from the proxy package as we are about to GC, and mesh merging creates assets that are 
-				// supposed to be standalone
-				CleanStandaloneAssetsInPackage(AssetsOuter);
+			// Link proxy to actor
+			UHLODProxy* PreviousProxy = LODActor->GetProxy();
+			Proxy->AddMesh(LODActor, MainMesh, UHLODProxy::GenerateKeyForActor(LODActor));
+			bDirtyPackage |= (LODActor->GetProxy() != PreviousProxy);
 
-				// Collect garbage to clean up old unreferenced data in the HLOD package
-				CollectGarbage(GARBAGE_COLLECTION_KEEPFLAGS);
+			if(bDirtyPackage)
+			{
+				LODActor->MarkPackageDirty();
 			}
 
-			// Add imposters
-			if (AllImposters.Num() > 0)
-			{
-				struct FLODImposterBatch
-				{
-					UStaticMesh*		StaticMesh;
-					TArray<FTransform>	Transforms;
-				};
+			// Clean out standalone meshes from the proxy package as we are about to GC, and mesh merging creates assets that are 
+			// supposed to be standalone
+			CleanStandaloneAssetsInPackage(AssetsOuter);
 
-				// Get all meshes + transforms for all imposters type (per material)
-				TMap<UMaterialInterface*, FLODImposterBatch> ImposterBatches;
-				for(UStaticMeshComponent* Imposter : AllImposters)
-				{
-					UMaterialInterface* Material = LODActor->GetImposterMaterial(Imposter);
-					check(Material);
+			// Collect garbage to clean up old unreferenced data in the HLOD package
+			CollectGarbage(GARBAGE_COLLECTION_KEEPFLAGS);
+		}
+	}
 
-					FLODImposterBatch& LODImposterBatch = ImposterBatches.FindOrAdd(Material);
-					LODImposterBatch.Transforms.Add(Imposter->GetOwner()->GetActorTransform());
+	// Add imposters
+	if (AllImposters.Num() > 0)
+	{
+		struct FLODImposterBatch
+		{
+			UStaticMesh*		StaticMesh;
+			TArray<FTransform>	Transforms;
+		};
 
-					// The static mesh hasn't been created yet, do it.
-					if(LODImposterBatch.StaticMesh == nullptr)
-					{
-						LODImposterBatch.StaticMesh = CreateImposterStaticMesh(Imposter, Material, LODSetup.ProxySetting);
-					}
-				}
+		// Get all meshes + transforms for all imposters type (per material)
+		TMap<UMaterialInterface*, FLODImposterBatch> ImposterBatches;
+		for (UStaticMeshComponent* Imposter : AllImposters)
+		{
+			UMaterialInterface* Material = LODActor->GetImposterMaterial(Imposter);
+			check(Material);
 
-				// Add imposters to the LODActor
-				for(const TPair<UMaterialInterface*, FLODImposterBatch> ImposterBatch : ImposterBatches)
-				{
-					LODActor->SetupImposters(ImposterBatch.Key, ImposterBatch.Value.StaticMesh, ImposterBatch.Value.Transforms);
-				}
+			FLODImposterBatch& LODImposterBatch = ImposterBatches.FindOrAdd(Material);
+			LODImposterBatch.Transforms.Add(Imposter->GetOwner()->GetActorTransform());
+
+			// The static mesh hasn't been created yet, do it.
+			if (LODImposterBatch.StaticMesh == nullptr)
+			{
+				LODImposterBatch.StaticMesh = CreateImposterStaticMesh(Imposter, Material, LODSetup.ProxySetting);
 			}
+		}
 
-			return true;
+		// Add imposters to the LODActor
+		for (const TPair<UMaterialInterface*, FLODImposterBatch> ImposterBatch : ImposterBatches)
+		{
+			LODActor->SetupImposters(ImposterBatch.Key, ImposterBatch.Value.StaticMesh, ImposterBatch.Value.Transforms);
 		}
 	}
-	return false;
+
+	return true;
 }
 
 bool FHierarchicalLODUtilities::BuildStaticMeshForLODActor(ALODActor* LODActor, UPackage* AssetsOuter, const FHierarchicalSimplification& LODSetup)
diff --git a/Engine/Source/Developer/MeshMergeUtilities/Private/MeshMergeUtilities.cpp b/Engine/Source/Developer/MeshMergeUtilities/Private/MeshMergeUtilities.cpp
--- a/Engine/Source/Developer/MeshMergeUtilities/Private/MeshMergeUtilities.cpp
+++ b/Engine/Source/Developer/MeshMergeUtilities/Private/MeshMergeUtilities.cpp
@@ -1767,17 +1767,18 @@ void FMeshMergeUtilities::MergeComponentsToStaticMesh(const TArray<UPrimitiveCom
 			else
 			{
 				StaticMeshComponentsToMerge.Add(MeshComponent);
+			}
 
-				// Save the pivot and asset package name of the first mesh, will later be used for creating merged mesh asset 
-				if (bFirstMesh)
-				{
-					// Mesh component pivot point
-					MergedAssetPivot = InSettings.bPivotPointAtZero ? FVector::ZeroVector : MeshComponent->GetComponentTransform().GetLocation();
-					// Source mesh asset package name
-					MergedAssetPackageName = MeshComponent->GetStaticMesh()->GetOutermost()->GetName();
+			// Save the pivot and asset package name of the first mesh, will later be used for creating merged mesh asset 
+			if (bFirstMesh)
+			{
+				// Mesh component pivot point
+				MergedAssetPivot = InSettings.bPivotPointAtZero ? FVector::ZeroVector : MeshComponent->GetComponentTransform().GetLocation();
 
-					bFirstMesh = false;
-				}
+				// Source mesh asset package name
+				MergedAssetPackageName = MeshComponent->GetStaticMesh()->GetOutermost()->GetName();
+
+				bFirstMesh = false;
 			}
 		}
 	}
@@ -1949,7 +1950,7 @@ void FMeshMergeUtilities::MergeComponentsToStaticMesh(const TArray<UPrimitiveCom
 	UMaterialInterface* MergedMaterial = nullptr;
 
 	// If the user wants to merge materials into a single one
-	if (bMergeMaterialData)
+	if (bMergeMaterialData && UniqueMaterials.Num() != 0)
 	{
 		UMaterialOptions* MaterialOptions = PopulateMaterialOptions(InSettings.MaterialSettings);
 		// Check each material to see if the shader actually uses vertex data and collect flags
@@ -2497,12 +2498,13 @@ void FMeshMergeUtilities::MergeComponentsToStaticMesh(const TArray<UPrimitiveCom
 			StaticMesh->LightMapCoordinateIndex = LightMapUVChannel;
 		}
 
+		const bool bContainsImposters = ImposterComponents.Num() > 0;
 		TArray<UMaterialInterface*> ImposterMaterials;
 		FBox ImposterBounds(EForceInit::ForceInit);
 		for (int32 LODIndex = 0; LODIndex < MergedRawMeshes.Num(); ++LODIndex)
 		{
 			FMeshDescription& MergedMeshLOD = MergedRawMeshes[LODIndex];
-			if (MergedMeshLOD.Vertices().Num() > 0)
+			if (MergedMeshLOD.Vertices().Num() > 0 || bContainsImposters)
 			{
 				FStaticMeshSourceModel& SrcModel = StaticMesh->AddSourceModel();
 				// Don't allow the engine to recalculate normals
@@ -2520,7 +2522,6 @@ void FMeshMergeUtilities::MergeComponentsToStaticMesh(const TArray<UPrimitiveCom
 					SrcModel.BuildSettings.DistanceFieldResolutionScale = 0.0f;
 				}
 
-				const bool bContainsImposters = ImposterComponents.Num() > 0;
 				if (bContainsImposters)
 				{
 					// Merge imposter meshes to rawmesh
diff --git a/Engine/Source/Runtime/Engine/Private/LODActor.cpp b/Engine/Source/Runtime/Engine/Private/LODActor.cpp
--- a/Engine/Source/Runtime/Engine/Private/LODActor.cpp
+++ b/Engine/Source/Runtime/Engine/Private/LODActor.cpp
@@ -586,13 +586,13 @@ const bool ALODActor::IsBuilt(bool bInForce/*=false*/) const
 		}
 
 		// No proxy mesh
-		if (Proxy == nullptr)
+		if (StaticMeshComponent->GetStaticMesh() && Proxy == nullptr)
 		{
 			return false;
 		}
 
 		// Mismatched key
-		if (!Proxy->ContainsDataForActor(this))
+		if (Proxy != nullptr && !Proxy->ContainsDataForActor(this))
 		{
 			return false;
 		}
@@ -875,11 +875,14 @@ void ALODActor::DetermineShadowingFlags()
 
 	for (const TPair<const UMaterialInterface*, UInstancedStaticMeshComponent*>& Component : ImpostersStaticMeshComponents)
 	{
-		Component.Value->CastShadow = false;
-		Component.Value->bCastStaticShadow = false;
-		Component.Value->bCastDynamicShadow = false;
-		Component.Value->bCastFarShadow = false;
-		Component.Value->MarkRenderStateDirty();
+		if (Component.Value)
+		{
+			Component.Value->CastShadow = false;
+			Component.Value->bCastStaticShadow = false;
+			Component.Value->bCastDynamicShadow = false;
+			Component.Value->bCastFarShadow = false;
+			Component.Value->MarkRenderStateDirty();
+		}
 	}
 
 	for (AActor* Actor : SubActors)
