HLOD batched imposter meshes have complex collision
#rnx

diff --git a/Engine/Source/Developer/HierarchicalLODUtilities/Private/HierarchicalLODUtilities.cpp b/Engine/Source/Developer/HierarchicalLODUtilities/Private/HierarchicalLODUtilities.cpp
--- a/Engine/Source/Developer/HierarchicalLODUtilities/Private/HierarchicalLODUtilities.cpp
+++ b/Engine/Source/Developer/HierarchicalLODUtilities/Private/HierarchicalLODUtilities.cpp
@@ -307,7 +307,12 @@ UStaticMesh* CreateImposterStaticMesh(UStaticMeshComponent* InComponent, UMateri
 		FMeshDescription* ImposterMesh = StaticMesh->CreateMeshDescription(0);
 		const IMeshMergeUtilities& MeshMergeUtilities = FModuleManager::Get().LoadModuleChecked<IMeshMergeModule>("MeshMergeUtilities").GetUtilities();
 		MeshMergeUtilities.ExtractImposterToRawMesh(InComponent, *ImposterMesh);
-	
+
+		// Disable collisions on imposters
+		FMeshSectionInfo Info = StaticMesh->GetSectionInfoMap().Get(0, 0);
+		Info.bEnableCollision = false;
+		StaticMesh->GetSectionInfoMap().Set(0, 0, Info);
+
 		// Commit mesh description and materials list to static mesh
 		StaticMesh->CommitMeshDescription(0);
 		StaticMesh->StaticMaterials = { InMaterial };
@@ -318,7 +323,7 @@ UStaticMesh* CreateImposterStaticMesh(UStaticMeshComponent* InComponent, UMateri
 		StaticMesh->PostEditChange();
 
 		// Our imposters meshes are flat, but they actually represent a volume.
-		// Extend the imposteor bounds using the original mesh bounds.
+		// Extend the imposter bounds using the original mesh bounds.
 		if (StaticMesh->GetBoundingBox().GetVolume() == 0)
 		{
 			const FBox StaticMeshBox = StaticMesh->GetBoundingBox();
