Crash trying to Rebuild HLODs

Fix possible crash if HLOD packages were duplicated

diff --git a/Engine/Source/Developer/HierarchicalLODUtilities/Private/HierarchicalLODUtilities.cpp b/Engine/Source/Developer/HierarchicalLODUtilities/Private/HierarchicalLODUtilities.cpp
--- a/Engine/Source/Developer/HierarchicalLODUtilities/Private/HierarchicalLODUtilities.cpp
+++ b/Engine/Source/Developer/HierarchicalLODUtilities/Private/HierarchicalLODUtilities.cpp
@@ -136,15 +136,17 @@ UHLODProxy* FHierarchicalLODUtilities::CreateOrRetrieveLevelHLODProxy(const ULev
 {
 	UPackage* HLODPackage = CreateOrRetrieveLevelHLODPackage(InLevel, HLODLevelIndex);
 
-	// check if our asset exists
+	// Check if our asset exists
 	const FString HLODProxyName = GetHLODProxyName(InLevel, HLODLevelIndex);
 	UHLODProxy* Proxy = FindObject<UHLODProxy>(HLODPackage, *HLODProxyName);
-	if(Proxy == nullptr)
+
+	// If proxy doesn't exist or is pointing to another world (could happen if package is duplicated)
+	if(Proxy == nullptr || Proxy->GetMap() != InLevel->GetWorld())
 	{
-		// make sure that the package doesnt have any standalone meshes etc. (i.e. this is an old style package)
+		// Make sure that the package doesnt have any standalone meshes etc. (i.e. this is an old style package)
 		CleanStandaloneAssetsInPackage(HLODPackage);
 
-		// create the new asset
+		// Create the new asset
 		Proxy = NewObject<UHLODProxy>(HLODPackage, *HLODProxyName, RF_Public | RF_Standalone);
 		Proxy->SetMap(UWorld::FindWorldInPackage(InLevel->GetOutermost()));
 	}
diff --git a/Engine/Source/Runtime/Engine/Classes/Engine/HLODProxy.h b/Engine/Source/Runtime/Engine/Classes/Engine/HLODProxy.h
--- a/Engine/Source/Runtime/Engine/Classes/Engine/HLODProxy.h
+++ b/Engine/Source/Runtime/Engine/Classes/Engine/HLODProxy.h
@@ -70,6 +70,9 @@ public:
 	/** Setup the map - only called at initial construction */
 	void SetMap(const UWorld* InMap);
 
+    /** Get the owner map for this HLOD proxy */
+	TSoftObjectPtr<UWorld> GetMap() const;
+
 	/** Adds a static mesh and the key used to generate it */
 	void AddMesh(ALODActor* InLODActor, UStaticMesh* InStaticMesh, const FName& InKey);
 
diff --git a/Engine/Source/Runtime/Engine/Private/HLODProxy.cpp b/Engine/Source/Runtime/Engine/Private/HLODProxy.cpp
--- a/Engine/Source/Runtime/Engine/Private/HLODProxy.cpp
+++ b/Engine/Source/Runtime/Engine/Private/HLODProxy.cpp
@@ -24,6 +24,11 @@ void UHLODProxy::SetMap(const UWorld* InMap)
 	OwningMap = InMap;
 }
 
+TSoftObjectPtr<UWorld> UHLODProxy::GetMap() const
+{
+	return OwningMap;
+}
+
 void UHLODProxy::AddMesh(ALODActor* InLODActor, UStaticMesh* InStaticMesh, const FName& InKey)
 {
 	InLODActor->Proxy = this;
