Fix bug where UHLODProxy were recreated every time the HLOD proxies were built from HLOD rendering maps.
This would cause unbuilt HLOD static meshes to be removed from the HLOD package.

diff --git a/Engine/Source/Developer/HierarchicalLODUtilities/Private/HierarchicalLODUtilities.cpp b/Engine/Source/Developer/HierarchicalLODUtilities/Private/HierarchicalLODUtilities.cpp
--- a/Engine/Source/Developer/HierarchicalLODUtilities/Private/HierarchicalLODUtilities.cpp
+++ b/Engine/Source/Developer/HierarchicalLODUtilities/Private/HierarchicalLODUtilities.cpp
@@ -140,15 +140,18 @@ UHLODProxy* FHierarchicalLODUtilities::CreateOrRetrieveLevelHLODProxy(const ULev
 	const FString HLODProxyName = GetHLODProxyName(InLevel, HLODLevelIndex);
 	UHLODProxy* Proxy = FindObject<UHLODProxy>(HLODPackage, *HLODProxyName);
 
+	// Get the world associated with this level
+	UWorld* LevelWorld = UWorld::FindWorldInPackage(InLevel->GetOutermost());
+
 	// If proxy doesn't exist or is pointing to another world (could happen if package is duplicated)
-	if(Proxy == nullptr || Proxy->GetMap() != InLevel->GetWorld())
+	if(Proxy == nullptr || Proxy->GetMap() != LevelWorld)
 	{
-		// Make sure that the package doesnt have any standalone meshes etc. (i.e. this is an old style package)
+		// Make sure that the package doesn't have any standalone meshes etc. (i.e. this is an old style package)
 		CleanStandaloneAssetsInPackage(HLODPackage);
 
 		// Create the new asset
 		Proxy = NewObject<UHLODProxy>(HLODPackage, *HLODProxyName, RF_Public | RF_Standalone);
-		Proxy->SetMap(UWorld::FindWorldInPackage(InLevel->GetOutermost()));
+		Proxy->SetMap(LevelWorld);
 	}
 
 	return Proxy;	
