Changed UStaticMeshComponent::GetStaticMesh to return a TObjectPtr as it showed up in profiles as the next most expensive contributor to redundant access tracking calls.
Fixed TObjectKey to be implicitly constructible from a wrapped object pointer the same way FObjectKey is.
Added a call-site fix for a const_cast from wrapped pointer to raw pointer.

diff --git a/Engine/Source/Developer/HierarchicalLODUtilities/Private/HierarchicalLODUtilities.cpp b/Engine/Source/Developer/HierarchicalLODUtilities/Private/HierarchicalLODUtilities.cpp
--- a/Engine/Source/Developer/HierarchicalLODUtilities/Private/HierarchicalLODUtilities.cpp
+++ b/Engine/Source/Developer/HierarchicalLODUtilities/Private/HierarchicalLODUtilities.cpp
@@ -1093,7 +1093,7 @@ int32 FHierarchicalLODUtilities::GetLODLevelForScreenSize(const UStaticMeshCompo
 	checkf(RenderData != nullptr, TEXT("StaticMesh in StaticMeshComponent %s contains invalid render data"), *StaticMeshComponent->GetName());
 	checkf(StaticMeshComponent->GetStaticMesh()->GetNumSourceModels() > 0, TEXT("StaticMesh in StaticMeshComponent %s contains no SourceModels"), *StaticMeshComponent->GetName());
 
-	return ComputeStaticMeshLODLevel(const_cast<const UStaticMesh*>(StaticMeshComponent->GetStaticMesh())->GetSourceModels(), RenderData, ScreenSize);
+	return ComputeStaticMeshLODLevel(const_cast<const UStaticMesh*>(ToRawPtr(StaticMeshComponent->GetStaticMesh()))->GetSourceModels(), RenderData, ScreenSize);
 }
 
 AHierarchicalLODVolume* FHierarchicalLODUtilities::CreateVolumeForLODActor(ALODActor* InLODActor, UWorld* InWorld)
diff --git a/Engine/Source/Runtime/CoreUObject/Public/UObject/ObjectKey.h b/Engine/Source/Runtime/CoreUObject/Public/UObject/ObjectKey.h
--- a/Engine/Source/Runtime/CoreUObject/Public/UObject/ObjectKey.h
+++ b/Engine/Source/Runtime/CoreUObject/Public/UObject/ObjectKey.h
@@ -111,8 +111,12 @@ public:
 	FORCEINLINE TObjectKey() = default;
 
 	/** Construct from an object pointer */
-	FORCEINLINE TObjectKey(const ElementType* Object)
-		: ObjectKey(Object)
+	template <
+		typename U,
+		decltype(ImplicitConv<const InElementType*>(std::declval<U>()))* = nullptr
+	>
+	FORCEINLINE TObjectKey(U Object)
+		: ObjectKey(ImplicitConv<const InElementType*>(Object))
 	{
 	}
 
diff --git a/Engine/Source/Runtime/Engine/Classes/Components/StaticMeshComponent.h b/Engine/Source/Runtime/Engine/Classes/Components/StaticMeshComponent.h
--- a/Engine/Source/Runtime/Engine/Classes/Components/StaticMeshComponent.h
+++ b/Engine/Source/Runtime/Engine/Classes/Components/StaticMeshComponent.h
@@ -383,7 +383,7 @@ public:
 	virtual bool SetStaticMesh(class UStaticMesh* NewMesh);
 
 	/** Get the StaticMesh used by this instance. */
-	UStaticMesh* GetStaticMesh() const 
+	TObjectPtr<UStaticMesh> GetStaticMesh() const 
 	{ 
 #if WITH_EDITOR
 		// This should never happen and is a last resort, we should have catched the property overwrite well before we reach this code
