Added HLODBatchingPolicy (enum) member to UPrimitiveComponent
* Replaces the confusing bUseMaxLODAsImposter & bBatchImpostersAsInstances
* Modified ALODActor::IsBuilt() to avoid pulling on editoronly data

diff --git a/Engine/Source/Developer/HierarchicalLODUtilities/Private/HierarchicalLODUtilities.cpp b/Engine/Source/Developer/HierarchicalLODUtilities/Private/HierarchicalLODUtilities.cpp
--- a/Engine/Source/Developer/HierarchicalLODUtilities/Private/HierarchicalLODUtilities.cpp
+++ b/Engine/Source/Developer/HierarchicalLODUtilities/Private/HierarchicalLODUtilities.cpp
@@ -436,7 +436,7 @@ bool FHierarchicalLODUtilities::BuildStaticMeshForLODActor(ALODActor* LODActor,
 		return false;
 	}
 
-	TArray<UStaticMeshComponent*> AllImposters;
+	TArray<UStaticMeshComponent*> AllInstances;
 	if (LODSetup.MergeSetting.bIncludeImposters)
 	{			
 		// Retrieve all imposters.
@@ -446,13 +446,13 @@ bool FHierarchicalLODUtilities::BuildStaticMeshForLODActor(ALODActor* LODActor,
 			{
 				if (LODActor->ShouldUseInstancing(StaticMeshComponent))
 				{
-					AllImposters.Add(StaticMeshComponent);
+					AllInstances.Add(StaticMeshComponent);
 				}
 			}
 		}
 
-		// Imposters won't be merged in the HLOD mesh
-		AllComponents.RemoveAll([&](UPrimitiveComponent* Component) { return AllImposters.Contains(Component); });
+		// Instances won't be merged in the HLOD mesh
+		AllComponents.RemoveAll([&](UPrimitiveComponent* Component) { return AllInstances.Contains(Component); });
 	}
 
 	if (AllComponents.Num() > 0)
@@ -598,54 +598,54 @@ bool FHierarchicalLODUtilities::BuildStaticMeshForLODActor(ALODActor* LODActor,
 
 	// Add imposters
 	LODActor->ClearInstances();
-	if (AllImposters.Num() > 0)
+	if (AllInstances.Num() > 0)
 	{
-		struct FLODImposterBatch
+		struct FLODInstanceBatch
 		{
 			UStaticMesh*		StaticMesh;
 			TArray<FTransform>	Transforms;
 		};
 
-		// Get all meshes + transforms for all imposters type (per material)
-		TMap<FHLODInstancingKey, FLODImposterBatch> ImposterBatches;
-		for (UStaticMeshComponent* Imposter : AllImposters)
+		// Get all meshes + transforms for all instances type (per material)
+		TMap<FHLODInstancingKey, FLODInstanceBatch> InstancesBatches;
+		for (UStaticMeshComponent* SMC : AllInstances)
 		{
-			UStaticMesh* StaticMesh = Imposter->GetStaticMesh();
+			UStaticMesh* StaticMesh = SMC->GetStaticMesh();
 			check(StaticMesh);
 
-			UMaterialInterface* ImposterMaterial = GetImposterMaterial(Imposter);
+			UMaterialInterface* InstanceMaterial = GetImposterMaterial(SMC);
 
-			FHLODInstancingKey Key(Imposter->GetStaticMesh(), ImposterMaterial);
+			FHLODInstancingKey Key(SMC->GetStaticMesh(), InstanceMaterial);
 			check(Key.IsValid());
 
-			FLODImposterBatch& LODImposterBatch = ImposterBatches.FindOrAdd(Key);
+			FLODInstanceBatch& LODInstanceBatch = InstancesBatches.FindOrAdd(Key);
 
 			// If we have an ISMC, ensure we include all its instances
-			if (UInstancedStaticMeshComponent* InstancedSMC = Cast<UInstancedStaticMeshComponent>(Imposter))
+			if (UInstancedStaticMeshComponent* InstancedSMC = Cast<UInstancedStaticMeshComponent>(SMC))
 			{
 				FTransform ActorTransformWS = InstancedSMC->GetOwner()->GetActorTransform();
 
-				LODImposterBatch.Transforms.Reserve(LODImposterBatch.Transforms.Num() + InstancedSMC->GetInstanceCount());
+				LODInstanceBatch.Transforms.Reserve(LODInstanceBatch.Transforms.Num() + InstancedSMC->GetInstanceCount());
 				for (const FInstancedStaticMeshInstanceData& InstanceData : InstancedSMC->PerInstanceSMData)
 				{
 					FTransform InstanceTransformWS = FTransform(InstanceData.Transform) * ActorTransformWS;
-					LODImposterBatch.Transforms.Add(InstanceTransformWS);
+					LODInstanceBatch.Transforms.Add(InstanceTransformWS);
 				}
 			}
 			else
 			{
-				LODImposterBatch.Transforms.Add(Imposter->GetOwner()->GetActorTransform());
+				LODInstanceBatch.Transforms.Add(SMC->GetOwner()->GetActorTransform());
 			}
 
 			// The static mesh hasn't been created yet, do it.
-			if (LODImposterBatch.StaticMesh == nullptr)
+			if (LODInstanceBatch.StaticMesh == nullptr)
 			{
-				LODImposterBatch.StaticMesh = CreateImposterStaticMesh(Imposter, LODSetup.ProxySetting);
+				LODInstanceBatch.StaticMesh = CreateImposterStaticMesh(SMC, LODSetup.ProxySetting);
 			}
 		}
 
 		// Add imposters to the LODActor
-		for (const auto& ImposterBatch : ImposterBatches)
+		for (const auto& ImposterBatch : InstancesBatches)
 		{
 			LODActor->AddInstances(ImposterBatch.Value.StaticMesh, ImposterBatch.Key.Material, ImposterBatch.Value.Transforms);
 		}
diff --git a/Engine/Source/Developer/MeshMergeUtilities/Private/MeshMergeUtilities.cpp b/Engine/Source/Developer/MeshMergeUtilities/Private/MeshMergeUtilities.cpp
--- a/Engine/Source/Developer/MeshMergeUtilities/Private/MeshMergeUtilities.cpp
+++ b/Engine/Source/Developer/MeshMergeUtilities/Private/MeshMergeUtilities.cpp
@@ -1561,7 +1561,7 @@ void FMeshMergeUtilities::CreateProxyMesh(const TArray<UStaticMeshComponent*>& I
 	TArray<UStaticMeshComponent*> StaticMeshComponents;
 	for (UStaticMeshComponent* StaticMeshComponent : ComponentsToMerge)
 	{
-		if (StaticMeshComponent->bUseMaxLODAsImposter)
+		if (StaticMeshComponent->HLODBatchingPolicy != EHLODBatchingPolicy::None)
 		{
 			ImposterMeshComponents.Add(StaticMeshComponent);
 		}
@@ -2168,7 +2168,7 @@ void FMeshMergeUtilities::MergeComponentsToStaticMesh(const TArray<UPrimitiveCom
 				}
 			}
 
-			if(MeshComponent->bUseMaxLODAsImposter && InSettings.bIncludeImposters)
+			if((MeshComponent->HLODBatchingPolicy != EHLODBatchingPolicy::None) && InSettings.bIncludeImposters)
 			{
 				ImposterComponents.Add(MeshComponent);
 			}
@@ -2227,7 +2227,7 @@ void FMeshMergeUtilities::MergeComponentsToStaticMesh(const TArray<UPrimitiveCom
 			const int32 NumLODs = [&]()
 			{
 				const int32 NumberOfLODsAvailable = Adapter.GetNumberOfLODs();
-				if (Component->bUseMaxLODAsImposter)
+				if (Component->HLODBatchingPolicy != EHLODBatchingPolicy::None)
 				{
 					return InSettings.bIncludeImposters ? NumberOfLODsAvailable : NumberOfLODsAvailable - 1;
 				}
@@ -2277,7 +2277,7 @@ void FMeshMergeUtilities::MergeComponentsToStaticMesh(const TArray<UPrimitiveCom
 			const int32 LODIndex = [&]()
 			{
 				int32 LowestDetailLOD = Adapter.GetNumberOfLODs() - 1;
-				if (Component->bUseMaxLODAsImposter && !InSettings.bIncludeImposters)
+				if (Component->HLODBatchingPolicy != EHLODBatchingPolicy::None && !InSettings.bIncludeImposters)
 				{
 					LowestDetailLOD = FMath::Max(0, LowestDetailLOD - 1);
 				}
diff --git a/Engine/Source/Runtime/Core/Public/UObject/UE5PrivateFrostyStreamObjectVersion.h b/Engine/Source/Runtime/Core/Public/UObject/UE5PrivateFrostyStreamObjectVersion.h
--- a/Engine/Source/Runtime/Core/Public/UObject/UE5PrivateFrostyStreamObjectVersion.h
+++ b/Engine/Source/Runtime/Core/Public/UObject/UE5PrivateFrostyStreamObjectVersion.h
@@ -12,6 +12,9 @@ struct CORE_API FUE5PrivateFrostyStreamObjectVersion
 		// Before any version changes were made
 		BeforeCustomVersionWasAdded = 0,
 
+		// Added HLODBatchingPolicy member to UPrimitiveComponent, which replaces the confusing bUseMaxLODAsImposter & bBatchImpostersAsInstances.
+		HLODBatchingPolicy,
+
 		// -----<new versions can be added above this line>-------------------------------------------------
 		VersionPlusOne,
 		LatestVersion = VersionPlusOne - 1
diff --git a/Engine/Source/Runtime/Engine/Classes/Components/PrimitiveComponent.h b/Engine/Source/Runtime/Engine/Classes/Components/PrimitiveComponent.h
--- a/Engine/Source/Runtime/Engine/Classes/Components/PrimitiveComponent.h
+++ b/Engine/Source/Runtime/Engine/Classes/Components/PrimitiveComponent.h
@@ -22,6 +22,7 @@
 #include "VT/RuntimeVirtualTextureEnum.h"
 #include "HitProxies.h"
 #include "Interfaces/Interface_AsyncCompilation.h"
+#include "HLOD/HLODBatchingPolicy.h"
 #include "PrimitiveComponent.generated.h"
 
 class AController;
@@ -269,23 +270,19 @@ public:
 	ELightmapType LightmapType;
 
 #if WITH_EDITORONLY_DATA
+	/** Whether to include this component in HLODs or not. */
+	UPROPERTY(EditAnywhere, AdvancedDisplay, BlueprintReadWrite, Category = HLOD, meta = (DisplayName = "Include Component in HLOD"))
+	uint8 bEnableAutoLODGeneration : 1;
+
 	/** Which specific HLOD levels this component should be excluded from */
-	UPROPERTY(EditAnywhere, AdvancedDisplay, BlueprintReadWrite, Category = HLOD)
+	UPROPERTY(EditAnywhere, AdvancedDisplay, BlueprintReadWrite, Category = HLOD, meta = (DisplayName = "Exclude from HLOD Levels", EditConditionHides, EditCondition = "bEnableAutoLODGeneration"))
 	TArray<int32> ExcludeForSpecificHLODLevels;
 
-	/** If true, and if World setting has bEnableHierarchicalLOD equal to true, then this component will be included when generating a Proxy mesh for the parent Actor */
-	UPROPERTY(EditAnywhere, AdvancedDisplay, BlueprintReadWrite, Category = HLOD, meta = (DisplayName = "Include Component for HLOD Mesh generation"))
-	uint8 bEnableAutoLODGeneration : 1;
+	/** Determines how the geometry of a component will be incorporated in proxy (simplified) HLODs. */
+	UPROPERTY(EditAnywhere, AdvancedDisplay, BlueprintReadWrite, Category = HLOD, meta = (DisplayName = "HLOD Batching Policy", EditConditionHides, EditCondition = "bEnableAutoLODGeneration"))
+	EHLODBatchingPolicy HLODBatchingPolicy;
 #endif 
 
-	/** Use the Maximum LOD Mesh (imposter) instead of including Mesh data from this component in the Proxy Generation process */
-	UPROPERTY(EditAnywhere, AdvancedDisplay, BlueprintReadWrite, Category = HLOD)
-	uint8 bUseMaxLODAsImposter : 1;
-
-	/** If true, the proxy generation process will use instancing to render this imposter */
-	UPROPERTY(EditAnywhere, AdvancedDisplay, BlueprintReadWrite, Category = HLOD, meta = (EditCondition = "bUseMaxLODAsImposter"))
-	uint8 bBatchImpostersAsInstances : 1;
-
 	/**
 	 * When enabled this object will not be culled by distance. This is ignored if a child of a HLOD.
 	 */
@@ -612,6 +609,13 @@ private:
 	UPROPERTY()
 	TEnumAsByte<enum ECanBeCharacterBase> CanBeCharacterBase_DEPRECATED;
 
+	/** Deprecated - represented by HLODBatchingPolicy == EHLODBatchingPolicy::MeshSection */
+	UPROPERTY()
+	uint8 bUseMaxLODAsImposter_DEPRECATED : 1;
+
+	/** Deprecated - represented by HLODBatchingPolicy == EHLODBatchingPolicy::Instancing */
+	UPROPERTY()
+	uint8 bBatchImpostersAsInstances_DEPRECATED : 1;
 #endif
 
 	FMaskFilter MoveIgnoreMask;
diff --git a/Engine/Source/Runtime/Engine/Classes/Engine/LODActor.h b/Engine/Source/Runtime/Engine/Classes/Engine/LODActor.h
--- a/Engine/Source/Runtime/Engine/Classes/Engine/LODActor.h
+++ b/Engine/Source/Runtime/Engine/Classes/Engine/LODActor.h
@@ -350,7 +350,7 @@ private:
 	void RegisterMeshComponents();
 	void UnregisterMeshComponents();
 
-#if (!(UE_BUILD_SHIPPING || UE_BUILD_TEST)) || WITH_EDITOR
+#if WITH_EDITOR
 	// Get/Create the LOD instanced static mesh component for a given imposter material.
 	UInstancedStaticMeshComponent* GetISMComponent(const FHLODInstancingKey& InstancingKey) const;
 	UInstancedStaticMeshComponent* GetOrCreateISMComponent(const FHLODInstancingKey& InstancingKey);
diff --git a/Engine/Source/Runtime/Engine/Private/Components/PrimitiveComponent.cpp b/Engine/Source/Runtime/Engine/Private/Components/PrimitiveComponent.cpp
--- a/Engine/Source/Runtime/Engine/Private/Components/PrimitiveComponent.cpp
+++ b/Engine/Source/Runtime/Engine/Private/Components/PrimitiveComponent.cpp
@@ -33,6 +33,7 @@
 #include "Algo/Copy.h"
 #include "UObject/RenderingObjectVersion.h"
 #include "UObject/FortniteMainBranchObjectVersion.h"
+#include "UObject/UE5PrivateFrostyStreamObjectVersion.h"
 #include "EngineModule.h"
 #include "UObject/ObjectSaveContext.h"
 
@@ -332,8 +333,11 @@ UPrimitiveComponent::UPrimitiveComponent(const FObjectInitializer& ObjectInitial
 
 	LDMaxDrawDistance = 0.f;
 	CachedMaxDrawDistance = 0.f;
-	bUseMaxLODAsImposter = false;
-	bBatchImpostersAsInstances = false;
+#if WITH_EDITORONLY_DATA
+	HLODBatchingPolicy = EHLODBatchingPolicy::None;
+	bUseMaxLODAsImposter_DEPRECATED = false;
+	bBatchImpostersAsInstances_DEPRECATED = false;
+#endif
 	bNeverDistanceCull = false;
 
 	bUseEditorCompositing = false;
@@ -933,6 +937,7 @@ void UPrimitiveComponent::Serialize(FArchive& Ar)
 	Super::Serialize(Ar);
 	Ar.UsingCustomVersion(FRenderingObjectVersion::GUID);
 	Ar.UsingCustomVersion(FFortniteMainBranchObjectVersion::GUID);
+	Ar.UsingCustomVersion(FUE5PrivateFrostyStreamObjectVersion::GUID);
 
 	// as temporary fix for the bug TTP 299926
 	// permanent fix is coming
@@ -948,6 +953,21 @@ void UPrimitiveComponent::Serialize(FArchive& Ar)
 			LightmapType = ELightmapType::ForceSurface;
 		}
 	}
+
+#if WITH_EDITORONLY_DATA
+	if (Ar.CustomVer(FUE5PrivateFrostyStreamObjectVersion::GUID) < FUE5PrivateFrostyStreamObjectVersion::HLODBatchingPolicy)
+	{
+		if (bUseMaxLODAsImposter_DEPRECATED)
+		{
+			HLODBatchingPolicy = EHLODBatchingPolicy::MeshSection;
+		}
+
+		if (bBatchImpostersAsInstances_DEPRECATED)
+		{
+			HLODBatchingPolicy = EHLODBatchingPolicy::Instancing;
+		}
+	}
+#endif
 }
 
 #if WITH_EDITOR
diff --git a/Engine/Source/Runtime/Engine/Private/HLODProxy.cpp b/Engine/Source/Runtime/Engine/Private/HLODProxy.cpp
--- a/Engine/Source/Runtime/Engine/Private/HLODProxy.cpp
+++ b/Engine/Source/Runtime/Engine/Private/HLODProxy.cpp
@@ -428,8 +428,7 @@ uint32 UHLODProxy::GetCRC(UStaticMeshComponent* InComponent, uint32 InCRC, const
 
 	// Include other relevant properties
 	Ar << InComponent->ForcedLodModel;
-	bool bUseMaxLODAsImposter = InComponent->bUseMaxLODAsImposter;
-	Ar << bUseMaxLODAsImposter;
+	Ar << InComponent->HLODBatchingPolicy;
 	bool bCastShadow = InComponent->CastShadow;
 	Ar << bCastShadow;
 	bool bCastStaticShadow = InComponent->bCastStaticShadow;
diff --git a/Engine/Source/Runtime/Engine/Private/LODActor.cpp b/Engine/Source/Runtime/Engine/Private/LODActor.cpp
--- a/Engine/Source/Runtime/Engine/Private/LODActor.cpp
+++ b/Engine/Source/Runtime/Engine/Private/LODActor.cpp
@@ -588,8 +588,9 @@ const bool ALODActor::IsBuilt(bool bInForce/*=false*/) const
 		{
 			if(SubActor)
 			{
-				UStaticMeshComponent* LODComponent = GetLODComponentForActor(SubActor, false);
-				if (LODComponent == nullptr || LODComponent->GetStaticMesh() == nullptr)
+				UStaticMeshComponent* SMComponent = SubActor->FindComponentByClass<UStaticMeshComponent>();
+				UStaticMeshComponent* LODComponent = SMComponent ? Cast<UStaticMeshComponent>(SMComponent->GetLODParentPrimitive()) : nullptr;
+				if (LODComponent == nullptr || LODComponent->GetOwner() != this || LODComponent->GetStaticMesh() == nullptr)
 				{
 					return false;
 				}
@@ -1095,15 +1096,11 @@ bool ALODActor::UpdateProxyDesc()
 	return false;
 }
 
-#endif // WITH_EDITOR
-
-#if (!(UE_BUILD_SHIPPING || UE_BUILD_TEST)) || WITH_EDITOR
-
 bool ALODActor::ShouldUseInstancing(const UStaticMeshComponent* InComponent)
 {
 	check(InComponent);
 
-	if (!InComponent->bUseMaxLODAsImposter || !InComponent->bBatchImpostersAsInstances)
+	if (InComponent->HLODBatchingPolicy != EHLODBatchingPolicy::Instancing)
 	{
 		return false;
 	}
@@ -1173,11 +1170,7 @@ static FHLODInstancingKey GetInstancingKey(const AActor* InActor, int32 InLODLev
 	InActor->GetComponents<UStaticMeshComponent>(Components);
 	Components.RemoveAll([&](UStaticMeshComponent* Val)
 	{
-#if WITH_EDITOR
 		return Val->GetStaticMesh() == nullptr || !Val->ShouldGenerateAutoLOD(InLODLevel - 1);
-#else
-		return Val->GetStaticMesh() == nullptr;
-#endif
 	});
 
 	if (Components.Num() == 1 && ALODActor::ShouldUseInstancing(Components[0]))
@@ -1263,7 +1256,7 @@ UStaticMeshComponent* ALODActor::GetOrCreateLODComponentForActor(const AActor* I
 	return LODComponent;
 }
 
-#endif // (!(UE_BUILD_SHIPPING || UE_BUILD_TEST)) || WITH_EDITOR
+#endif // WITH_EDITOR
 
 FBox ALODActor::GetComponentsBoundingBox(bool bNonColliding, bool bIncludeFromChildActors) const
 {
diff --git a/Engine/Source/Runtime/Engine/Public/HLOD/HLODBatchingPolicy.h b/Engine/Source/Runtime/Engine/Public/HLOD/HLODBatchingPolicy.h
new file mode 100644
--- /dev/null
+++ b/Engine/Source/Runtime/Engine/Public/HLOD/HLODBatchingPolicy.h
@@ -0,0 +1,19 @@
+// Copyright Epic Games, Inc. All Rights Reserved.
+
+#pragma once
+
+#include "CoreMinimal.h"
+
+/** Determines how the geometry of a component will be incorporated in proxy (simplified) HLODs. */
+UENUM()
+enum class EHLODBatchingPolicy : uint8
+{
+	/** No batching to be performed, geometry is to be simplified. */ 
+	None,
+
+	/** Batch this component geometry (using the lowest LOD) as a separate mesh section, grouping by material. */
+	MeshSection,
+
+	/** Batch this component geometry (using the lowest LOD) as a separate instanced static mesh component in the generated actor. */
+	Instancing
+};
