HLOD: Added support for Approximate simplification method to Level HLODs
* Switched FHierarchicalSimplification::bSimplifyMesh to an enum FHierarchicalSimplification::SimplificationMethod
* Cleaned deprecation handling for a few classes

diff --git a/Engine/Source/Developer/HierarchicalLODUtilities/HierarchicalLODUtilities.Build.cs b/Engine/Source/Developer/HierarchicalLODUtilities/HierarchicalLODUtilities.Build.cs
--- a/Engine/Source/Developer/HierarchicalLODUtilities/HierarchicalLODUtilities.Build.cs
+++ b/Engine/Source/Developer/HierarchicalLODUtilities/HierarchicalLODUtilities.Build.cs
@@ -21,6 +21,7 @@ public class HierarchicalLODUtilities : ModuleRules
 				"BSPUtils",
 				"EditorFramework",
 				"Engine",
+				"MaterialUtilities",
 				"MeshDescription",
 				"StaticMeshDescription",
 				"UnrealEd",
@@ -28,10 +29,11 @@ public class HierarchicalLODUtilities : ModuleRules
 			}
         );
 
-        PrivateIncludePaths.AddRange(
+		PrivateIncludePathModuleNames.AddRange(
             new string[]
             {
-            }
+				"GeometryProcessingInterfaces"
+			}
         );
 
         DynamicallyLoadedModuleNames.AddRange(
@@ -40,7 +42,8 @@ public class HierarchicalLODUtilities : ModuleRules
                 "MeshUtilities",
                 "MeshMergeUtilities",
                 "MeshReductionInterface",
-            }
+				"GeometryProcessingInterfaces"
+			}
         );
 	}
 }
diff --git a/Engine/Source/Developer/HierarchicalLODUtilities/Private/HierarchicalLODUtilities.cpp b/Engine/Source/Developer/HierarchicalLODUtilities/Private/HierarchicalLODUtilities.cpp
--- a/Engine/Source/Developer/HierarchicalLODUtilities/Private/HierarchicalLODUtilities.cpp
+++ b/Engine/Source/Developer/HierarchicalLODUtilities/Private/HierarchicalLODUtilities.cpp
@@ -37,14 +37,21 @@
 #include "Settings/EditorExperimentalSettings.h"
 #endif // WITH_EDITOR
 
+#include "MaterialUtilities.h"
+#include "IGeometryProcessingInterfacesModule.h"
+#include "GeometryProcessingInterfaces/ApproximateActors.h"
+
 #include "HierarchicalLODProxyProcessor.h"
 #include "IMeshReductionManagerModule.h"
 #include "MeshMergeModule.h"
+#include "Algo/ForEach.h"
 #include "Algo/Transform.h"
+#include "ObjectTools.h"
 #include "Engine/HLODProxy.h"
 #include "HierarchicalLOD.h"
 #include "LevelUtils.h"
 #include "Materials/Material.h"
+#include "Materials/MaterialInstance.h"
 
 DEFINE_LOG_CATEGORY_STATIC(LogHierarchicalLODUtilities, Verbose, All);
 
@@ -409,6 +416,281 @@ static UStaticMesh* CreateImposterStaticMesh(UStaticMeshComponent* InComponent,
 	return ImposterStaticMesh;
 }
 
+
+struct FHLODBuildParams
+{
+	ALODActor*								LODActor;
+	UHLODProxy*								Proxy;
+	const TArray<UPrimitiveComponent*>&		Components;
+	const FHierarchicalSimplification&		LODSetup;
+	UMaterialInterface*						BaseMaterial;
+	UPackage*								AssetsOuter;
+	FString									PackageName;
+
+	FHLODBuildParams(const TArray<UPrimitiveComponent*>& InComponents, const FHierarchicalSimplification& InLODSetup)
+		: Components(InComponents)
+		, LODSetup(InLODSetup)
+	{
+	}
+};
+
+struct FHLODBuildResults
+{
+	bool			bDeferredResults = false;
+	UStaticMesh*	HLODMesh = nullptr;
+	FVector			HLODLocation = FVector::ZeroVector;
+};
+
+FHLODBuildResults GenerateHLODMesh_Simplify(const FHLODBuildParams& InBuildParams)
+{
+	FHLODBuildResults HLODBuildResults;
+
+	const IMeshMergeUtilities& MeshMergeUtilities = FModuleManager::Get().LoadModuleChecked<IMeshMergeModule>("MeshMergeUtilities").GetUtilities();
+
+	// Generate proxy mesh and proxy material assets
+	IMeshReductionManagerModule& MeshReductionModule = FModuleManager::Get().LoadModuleChecked<IMeshReductionManagerModule>("MeshReductionInterface");
+	const bool bHasMeshReductionCapableModule = (MeshReductionModule.GetMeshMergingInterface() != nullptr);
+
+	if (!bHasMeshReductionCapableModule)
+	{
+		return HLODBuildResults;
+	}
+
+	FHierarchicalLODUtilitiesModule& Module = FModuleManager::LoadModuleChecked<FHierarchicalLODUtilitiesModule>("HierarchicalLODUtilities");
+	FHierarchicalLODProxyProcessor* Processor = Module.GetProxyProcessor();
+
+	FHierarchicalSimplification OverrideLODSetup = InBuildParams.LODSetup;
+
+	FMeshProxySettings ProxySettings = InBuildParams.LODSetup.ProxySetting;
+	if (InBuildParams.LODActor->bOverrideMaterialMergeSettings)
+	{
+		ProxySettings.MaterialSettings = InBuildParams.LODActor->MaterialSettings;
+	}
+
+	if (InBuildParams.LODActor->bOverrideScreenSize)
+	{
+		ProxySettings.ScreenSize = InBuildParams.LODActor->ScreenSize;
+	}
+
+	if (InBuildParams.LODActor->bOverrideTransitionScreenSize)
+	{
+		OverrideLODSetup.TransitionScreenSize = InBuildParams.LODActor->TransitionScreenSize;
+	}
+
+	FGuid JobID = Processor->AddProxyJob(InBuildParams.LODActor, InBuildParams.Proxy, OverrideLODSetup);
+
+	TArray<UStaticMeshComponent*> StaticMeshComponents;
+	Algo::Transform(InBuildParams.Components, StaticMeshComponents, [](UPrimitiveComponent* InPrimitiveComponent) { return Cast<UStaticMeshComponent>(InPrimitiveComponent); });
+
+	MeshMergeUtilities.CreateProxyMesh(StaticMeshComponents, ProxySettings, InBuildParams.BaseMaterial, InBuildParams.AssetsOuter, "", JobID, Processor->GetCallbackDelegate(), true, OverrideLODSetup.TransitionScreenSize);
+
+	
+	HLODBuildResults.bDeferredResults = true;
+	
+	return HLODBuildResults;
+}
+
+FHLODBuildResults GenerateHLODMesh_Merge(const FHLODBuildParams& InBuildParams)
+{
+	const IMeshMergeUtilities& MeshMergeUtilities = FModuleManager::Get().LoadModuleChecked<IMeshMergeModule>("MeshMergeUtilities").GetUtilities();
+
+	FMeshMergingSettings MergeSettings = InBuildParams.LODSetup.MergeSetting;
+	if (InBuildParams.LODActor->bOverrideMaterialMergeSettings)
+	{
+		MergeSettings.MaterialSettings = InBuildParams.LODActor->MaterialSettings;
+	}
+
+	// update LOD parents before rebuild to ensure they are valid when mesh merge extensions are called.
+	InBuildParams.LODActor->UpdateSubActorLODParents();
+
+	FHLODBuildResults HLODBuildResults;
+
+	TArray<UObject*> OutAssets;
+	MeshMergeUtilities.MergeComponentsToStaticMesh(InBuildParams.Components, InBuildParams.LODActor->GetWorld(), MergeSettings, InBuildParams.BaseMaterial, InBuildParams.AssetsOuter, "", OutAssets, HLODBuildResults.HLODLocation, InBuildParams.LODSetup.TransitionScreenSize, true);
+
+	// set staticmesh
+	for (UObject* Asset : OutAssets)
+	{
+		if (UStaticMesh* StaticMesh = Cast<UStaticMesh>(Asset))
+		{
+			HLODBuildResults.HLODMesh = StaticMesh;
+			break;
+		}
+	}
+
+	return HLODBuildResults;
+}
+
+FHLODBuildResults GenerateHLODMesh_Approximate(const FHLODBuildParams& InBuildParams)
+{
+	FHLODBuildResults HLODBuildResults;
+
+	IGeometryProcessingInterfacesModule* GeomProcInterfaces = FModuleManager::Get().LoadModulePtr<IGeometryProcessingInterfacesModule>("GeometryProcessingInterfaces");
+	IGeometryProcessing_ApproximateActors* ApproxActorsAPI = GeomProcInterfaces ? GeomProcInterfaces->GetApproximateActorsImplementation() : nullptr;
+	const bool bHasApproximateActorsModule = ApproxActorsAPI != nullptr;
+
+	if (!ApproxActorsAPI)
+	{
+		return HLODBuildResults;
+	}
+
+	const FMeshApproximationSettings& UseSettings = InBuildParams.LODSetup.ApproximateSettings;
+	UMaterialInterface* HLODMaterial = InBuildParams.BaseMaterial;
+
+	IGeometryProcessing_ApproximateActors::FOptions Options = ApproxActorsAPI->ConstructOptions(UseSettings);
+	Options.bGenerateLightmapUVs = false;
+	Options.bCreatePhysicsBody = false;
+
+	// Material baking settings
+	Options.BakeMaterial = HLODMaterial;
+	if (!FMaterialUtilities::IsValidFlattenMaterial(Options.BakeMaterial))
+	{
+		Options.BakeMaterial = GEngine->DefaultFlattenMaterial;
+	}
+	Options.BaseColorTexParamName = FName(FMaterialUtilities::GetFlattenMaterialTextureName(EFlattenMaterialProperties::Diffuse, Options.BakeMaterial));
+	Options.NormalTexParamName = FName(FMaterialUtilities::GetFlattenMaterialTextureName(EFlattenMaterialProperties::Normal, Options.BakeMaterial));
+	Options.MetallicTexParamName = FName(FMaterialUtilities::GetFlattenMaterialTextureName(EFlattenMaterialProperties::Metallic, Options.BakeMaterial));
+	Options.RoughnessTexParamName = FName(FMaterialUtilities::GetFlattenMaterialTextureName(EFlattenMaterialProperties::Roughness, Options.BakeMaterial));
+	Options.SpecularTexParamName = FName(FMaterialUtilities::GetFlattenMaterialTextureName(EFlattenMaterialProperties::Specular, Options.BakeMaterial));
+	Options.EmissiveTexParamName = FName(FMaterialUtilities::GetFlattenMaterialTextureName(EFlattenMaterialProperties::Emissive, Options.BakeMaterial));
+	Options.bUsePackedMRS = true;
+	Options.PackedMRSTexParamName = FName("PackedTexture");
+
+	// Gather bounds of the input components
+	auto GetActorsBounds = [&]() -> FBoxSphereBounds
+	{
+		FBoxSphereBounds Bounds;
+		bool bFirst = true;
+
+		for (UPrimitiveComponent* Component : InBuildParams.Components)
+		{
+			FBoxSphereBounds ComponentBounds = Component->Bounds;
+			Bounds = bFirst ? ComponentBounds : Bounds + ComponentBounds;
+			bFirst = false;
+		}
+
+		return Bounds;
+	};
+
+	// Compute texel density if needed, depending on the TextureSizingType setting
+	ETextureSizingType TextureSizingType = UseSettings.MaterialSettings.TextureSizingType;
+	float TexelDensityPerMeter = 0.0f;
+
+	IGeometryProcessing_ApproximateActors::ETextureSizePolicy TextureSizePolicy = IGeometryProcessing_ApproximateActors::ETextureSizePolicy::TextureSize;
+	if (TextureSizingType == ETextureSizingType::TextureSizingType_AutomaticFromTexelDensity)
+	{
+		TexelDensityPerMeter = UseSettings.MaterialSettings.TargetTexelDensityPerMeter;
+		TextureSizePolicy = IGeometryProcessing_ApproximateActors::ETextureSizePolicy::TexelDensity;
+	}
+	else if (TextureSizingType == ETextureSizingType::TextureSizingType_AutomaticFromMeshScreenSize)
+	{
+		TexelDensityPerMeter = FMaterialUtilities::ComputeRequiredTexelDensityFromScreenSize(UseSettings.MaterialSettings.MeshMaxScreenSizePercent, GetActorsBounds().SphereRadius);
+		TextureSizePolicy = IGeometryProcessing_ApproximateActors::ETextureSizePolicy::TexelDensity;
+	}
+	else if (TextureSizingType == ETextureSizingType::TextureSizingType_AutomaticFromMeshDrawDistance)
+	{
+		TexelDensityPerMeter = FMaterialUtilities::ComputeRequiredTexelDensityFromDrawDistance(UseSettings.MaterialSettings.MeshMinDrawDistance, GetActorsBounds().SphereRadius);
+		TextureSizePolicy = IGeometryProcessing_ApproximateActors::ETextureSizePolicy::TexelDensity;
+	}
+
+	Options.MeshTexelDensity = TexelDensityPerMeter;
+	Options.TextureSizePolicy = TextureSizePolicy;
+
+	// Use temp packages - otherwise Approximate Actors will create it's mesh using the same name as the HLOD Proxy object.
+	const FString NewAssetNamePrefix(TEXT("NEWASSET_"));
+	FString PackageName = InBuildParams.AssetsOuter->GetPackage()->GetName();
+	FString PackagePath = FPackageName::GetLongPackagePath(PackageName);
+	FString AssetName = FPackageName::GetLongPackageAssetName(PackageName);
+	Options.BasePackagePath = PackagePath / NewAssetNamePrefix + AssetName;
+
+	// run actor approximation computation
+	TSet<AActor*> Actors;
+	Algo::Transform(InBuildParams.Components, Actors, [](UPrimitiveComponent* PrimitiveComponent) { return PrimitiveComponent->GetOwner(); });
+
+	IGeometryProcessing_ApproximateActors::FResults Results;
+	ApproxActorsAPI->ApproximateActors(Actors.Array(), Options, Results);
+
+	auto RenameNewAsset = [&PackagePath, &NewAssetNamePrefix, &InBuildParams](UObject* NewAsset)
+	{
+		FString AssetName = NewAsset->GetName();
+		AssetName.RemoveFromStart(NewAssetNamePrefix);
+			
+		// Add a prefix to the generated static mesh, otherwise it would share the same name as the HLOD Proxy object.
+		if (NewAsset->IsA<UStaticMesh>())
+		{
+			AssetName += TEXT("_Mesh");
+		}
+
+		// Remplace existing asset
+		UObject* AssetToReplace = StaticFindObjectFast(UObject::StaticClass(), InBuildParams.AssetsOuter, *AssetName);
+		if (AssetToReplace)
+		{
+			// Replace references
+			TArray<UObject*> ObjectsToReplace(&AssetToReplace, 1);
+			ObjectTools::ForceReplaceReferences(NewAsset, ObjectsToReplace);
+
+			// Move the previous asset to the transient package
+			AssetToReplace->Rename(nullptr, GetTransientPackage(), REN_DontCreateRedirectors | REN_NonTransactional | REN_ForceNoResetLoaders);
+		}
+			
+		NewAsset->Rename(*AssetName, InBuildParams.AssetsOuter, REN_NonTransactional | REN_DontCreateRedirectors | REN_ForceNoResetLoaders);
+	};
+
+	Algo::ForEach(Results.NewMeshAssets, RenameNewAsset);
+	Algo::ForEach(Results.NewMaterials, RenameNewAsset);
+	Algo::ForEach(Results.NewTextures, RenameNewAsset);
+	
+	if (Results.ResultCode == IGeometryProcessing_ApproximateActors::EResultCode::Success)
+	{
+		// Retrieve staticmesh
+		if (ensure(Results.NewMeshAssets.Num() == 1))
+		{
+			HLODBuildResults.HLODMesh = Results.NewMeshAssets[0];
+		}
+
+		// Setup material switches
+		for (UMaterialInterface* Material : Results.NewMaterials)
+		{
+			UMaterialInstance* MaterialInst = CastChecked<UMaterialInstance>(Material);
+
+			FStaticParameterSet StaticParameterSet;
+
+			auto SetStaticSwitch = [&StaticParameterSet](FName ParamName, bool bSet)
+			{
+				if (bSet)
+				{
+					FStaticSwitchParameter SwitchParameter;
+					SwitchParameter.ParameterInfo.Name = ParamName;
+					SwitchParameter.Value = true;
+					SwitchParameter.bOverride = true;
+					StaticParameterSet.EditorOnly.StaticSwitchParameters.Add(SwitchParameter);
+				}
+			};
+
+			// Set proper switches needed by our base flatten material
+			SetStaticSwitch("UseBaseColor", Options.bBakeBaseColor);
+			SetStaticSwitch("UseDiffuse", Options.bBakeBaseColor);
+			SetStaticSwitch("UseRoughness", Options.bBakeRoughness);
+			SetStaticSwitch("UseMetallic", Options.bBakeMetallic);
+			SetStaticSwitch("UseSpecular", Options.bBakeSpecular);
+			SetStaticSwitch("UseEmissive", Options.bBakeEmissive);
+			SetStaticSwitch("UseEmissiveHDR", Options.bBakeEmissive);
+			SetStaticSwitch("UseNormal", Options.bBakeNormalMap);
+			SetStaticSwitch("PackMetallic", Options.bUsePackedMRS);
+			SetStaticSwitch("PackSpecular", Options.bUsePackedMRS);
+			SetStaticSwitch("PackRoughness", Options.bUsePackedMRS);
+
+			// Force initializing the static permutations according to the switches we have set
+			MaterialInst->UpdateStaticPermutation(StaticParameterSet);
+			MaterialInst->InitStaticPermutation();
+			MaterialInst->PostEditChange();
+		}
+	}
+
+	return HLODBuildResults;
+}
+
+
 bool FHierarchicalLODUtilities::BuildStaticMeshForLODActor(ALODActor* LODActor, UHLODProxy* Proxy, const FHierarchicalSimplification& LODSetup, UMaterialInterface* InBaseMaterial)
 {
 	if (!Proxy || !LODActor)
@@ -457,84 +739,40 @@ bool FHierarchicalLODUtilities::BuildStaticMeshForLODActor(ALODActor* LODActor,
 
 	if (AllComponents.Num() > 0)
 	{
-		TArray<UObject*> OutAssets;
-		FVector OutProxyLocation = FVector::ZeroVector;
-		UStaticMesh* MainMesh = nullptr;
-
-		// Generate proxy mesh and proxy material assets
-		IMeshReductionManagerModule& MeshReductionModule = FModuleManager::Get().LoadModuleChecked<IMeshReductionManagerModule>("MeshReductionInterface");
-		const bool bHasMeshReductionCapableModule = (MeshReductionModule.GetMeshMergingInterface() != NULL);
-
-		const IMeshMergeUtilities& MeshMergeUtilities = FModuleManager::Get().LoadModuleChecked<IMeshMergeModule>("MeshMergeUtilities").GetUtilities();
-		
-		// Should give a unique name, so use the LODActor tag, or if empty, the first actor name
-		FString LODActorTag = LODActor->GetLODActorTag();
-		if (LODActorTag.IsEmpty())
+		FHLODBuildParams HLODBuildParams(AllComponents, LODSetup);
+		HLODBuildParams.LODActor = LODActor;
+		HLODBuildParams.Proxy = Proxy;
+		HLODBuildParams.BaseMaterial = InBaseMaterial;
+		HLODBuildParams.AssetsOuter = AssetsOuter;
+
+		FHLODBuildResults HLODBuildResults;
+		switch (LODSetup.SimplificationMethod)
 		{
-			const AActor* FirstActor = UHLODProxy::FindFirstActor(LODActor);
-			LODActorTag = *FirstActor->GetName();
-		}
-		const FString PackageName = FString::Printf(TEXT("LOD_%s_%i_%s"), *(AssetsOuter->GetName()), LODActor->LODLevel - 1, *LODActorTag);
-
-		if (bHasMeshReductionCapableModule && LODSetup.bSimplifyMesh)
-		{
-			FHierarchicalLODUtilitiesModule& Module = FModuleManager::LoadModuleChecked<FHierarchicalLODUtilitiesModule>("HierarchicalLODUtilities");
-			FHierarchicalLODProxyProcessor* Processor = Module.GetProxyProcessor();
-
-			FHierarchicalSimplification OverrideLODSetup = LODSetup;
-
-			FMeshProxySettings ProxySettings = LODSetup.ProxySetting;
-			if (LODActor->bOverrideMaterialMergeSettings)
-			{
-				ProxySettings.MaterialSettings = LODActor->MaterialSettings;
-			}
-
-			if (LODActor->bOverrideScreenSize)
-			{
-				ProxySettings.ScreenSize = LODActor->ScreenSize;
-			}
-
-			if (LODActor->bOverrideTransitionScreenSize)
-			{
-				OverrideLODSetup.TransitionScreenSize = LODActor->TransitionScreenSize;
-			}
-
-			FGuid JobID = Processor->AddProxyJob(LODActor, Proxy, OverrideLODSetup);
-
-			TArray<UStaticMeshComponent*> StaticMeshComponents;
-			Algo::Transform(AllComponents, StaticMeshComponents, [](UPrimitiveComponent* InPrimitiveComponent) { return Cast<UStaticMeshComponent>(InPrimitiveComponent); });
+		case EHierarchicalSimplificationMethod::Merge:
+			HLODBuildResults = GenerateHLODMesh_Merge(HLODBuildParams);
+			break;
+		case EHierarchicalSimplificationMethod::Simplify:
+			HLODBuildResults = GenerateHLODMesh_Simplify(HLODBuildParams);
+			break;
+		case EHierarchicalSimplificationMethod::Approximate:
+			HLODBuildResults = GenerateHLODMesh_Approximate(HLODBuildParams);
+			break;
 
-			MeshMergeUtilities.CreateProxyMesh(StaticMeshComponents, ProxySettings, InBaseMaterial, AssetsOuter, PackageName, JobID, Processor->GetCallbackDelegate(), true, OverrideLODSetup.TransitionScreenSize);
+		default:
+			UE_LOG(LogHierarchicalLODUtilities, Error, TEXT("Unsupported simplification method provided"));
+			return false;
 		}
-		else
+		
+		if (!HLODBuildResults.bDeferredResults)
 		{
-			FMeshMergingSettings MergeSettings = LODSetup.MergeSetting;
-			if (LODActor->bOverrideMaterialMergeSettings)
-			{
-				MergeSettings.MaterialSettings = LODActor->MaterialSettings;
-			}
-
-			// update LOD parents before rebuild to ensure they are valid when mesh merge extensions are called.
-			LODActor->UpdateSubActorLODParents();
-
-			MeshMergeUtilities.MergeComponentsToStaticMesh(AllComponents, LODActor->GetWorld(), MergeSettings, InBaseMaterial, AssetsOuter, PackageName, OutAssets, OutProxyLocation, LODSetup.TransitionScreenSize, true);
-
-			// set staticmesh
-			for (UObject* Asset : OutAssets)
-			{
-				UStaticMesh* StaticMesh = Cast<UStaticMesh>(Asset);
-
-				if (StaticMesh)
-				{
-					MainMesh = StaticMesh;
-				}
-			}
-
-			if (!MainMesh)
+			if (!HLODBuildResults.HLODMesh)
 			{
+				UE_LOG(LogHierarchicalLODUtilities, Warning, TEXT("No HLOD mesh generated"));
 				return false;
 			}
 
+			UStaticMesh* MainMesh = HLODBuildResults.HLODMesh;
+
 			// make sure the mesh won't affect navmesh generation
 			MainMesh->MarkAsNotHavingNavigationData();
 
@@ -542,8 +780,8 @@ bool FHierarchicalLODUtilities::BuildStaticMeshForLODActor(ALODActor* LODActor,
 			UStaticMesh* PreviousStaticMesh = LODActor->GetStaticMeshComponent()->GetStaticMesh();
 			bDirtyPackage |= (MainMesh != PreviousStaticMesh);
 			LODActor->SetStaticMesh(MainMesh);
-			bDirtyPackage |= (LODActor->GetActorLocation() != OutProxyLocation);
-			LODActor->SetActorLocation(OutProxyLocation);
+			bDirtyPackage |= (LODActor->GetActorLocation() != HLODBuildResults.HLODLocation);
+			LODActor->SetActorLocation(HLODBuildResults.HLODLocation);
 
 			// Check resulting mesh and give a warning if it exceeds the vertex / triangle cap for certain platforms
 			FProjectStatus ProjectStatus;
diff --git a/Engine/Source/Developer/MaterialUtilities/Private/ProxyMaterialUtilities.cpp b/Engine/Source/Developer/MaterialUtilities/Private/ProxyMaterialUtilities.cpp
--- a/Engine/Source/Developer/MaterialUtilities/Private/ProxyMaterialUtilities.cpp
+++ b/Engine/Source/Developer/MaterialUtilities/Private/ProxyMaterialUtilities.cpp
@@ -135,6 +135,11 @@ static FString GetMatchingParamName(EFlattenMaterialProperties InProperty, UMate
 	return ParamName;
 }
 
+FString FMaterialUtilities::GetFlattenMaterialTextureName(EFlattenMaterialProperties InProperty, UMaterialInterface* InBaseMaterial)
+{
+	return GetMatchingParamName(InProperty, InBaseMaterial) + TEXT("Texture");
+}
+
 // Validate that the provided material has all the required parameters needed to be considered a flattening material
 // If not, report what is missing
 bool FMaterialUtilities::IsValidFlattenMaterial(const UMaterialInterface* InBaseMaterial)
diff --git a/Engine/Source/Developer/MaterialUtilities/Public/MaterialUtilities.h b/Engine/Source/Developer/MaterialUtilities/Public/MaterialUtilities.h
--- a/Engine/Source/Developer/MaterialUtilities/Public/MaterialUtilities.h
+++ b/Engine/Source/Developer/MaterialUtilities/Public/MaterialUtilities.h
@@ -494,6 +494,9 @@ public:
 	/** Validate that the provided material has all the required parameters needed to be considered a flattening material */
 	static bool IsValidFlattenMaterial(const UMaterialInterface* InBaseMaterial);
 
+	/** Get the name of the Texture parameter associated with a given flatten material property */
+	static FString GetFlattenMaterialTextureName(EFlattenMaterialProperties InProperty, UMaterialInterface* InBaseMaterial);
+
 private:
 	
 	/**
diff --git a/Engine/Source/Editor/DetailCustomizations/DetailCustomizations.Build.cs b/Engine/Source/Editor/DetailCustomizations/DetailCustomizations.Build.cs
--- a/Engine/Source/Editor/DetailCustomizations/DetailCustomizations.Build.cs
+++ b/Engine/Source/Editor/DetailCustomizations/DetailCustomizations.Build.cs
@@ -81,7 +81,8 @@ public class DetailCustomizations : ModuleRules
 				"GraphEditor",
 				"MeshMergeUtilities",
 				"MeshReductionInterface",
-            }
+				"GeometryProcessingInterfaces"
+			}
 		);
 
 		DynamicallyLoadedModuleNames.AddRange(
@@ -90,7 +91,8 @@ public class DetailCustomizations : ModuleRules
 				"GameProjectGeneration",
 				"MeshMergeUtilities",
 				"MeshReductionInterface",
-            }
+				"GeometryProcessingInterfaces"
+			}
 		);
 	}
 }
diff --git a/Engine/Source/Editor/DetailCustomizations/Private/HierarchicalSimplificationCustomizations.cpp b/Engine/Source/Editor/DetailCustomizations/Private/HierarchicalSimplificationCustomizations.cpp
--- a/Engine/Source/Editor/DetailCustomizations/Private/HierarchicalSimplificationCustomizations.cpp
+++ b/Engine/Source/Editor/DetailCustomizations/Private/HierarchicalSimplificationCustomizations.cpp
@@ -1,14 +1,18 @@
 // Copyright Epic Games, Inc. All Rights Reserved.
 
 #include "HierarchicalSimplificationCustomizations.h"
-#include "Modules/ModuleManager.h"
 #include "GameFramework/WorldSettings.h"
 #include "IDetailChildrenBuilder.h"
 #include "DetailWidgetRow.h"
 #include "IDetailGroup.h"
 #include "IDetailPropertyRow.h"
+#include "PropertyRestriction.h"
 #include "MeshUtilities.h"
+
+#include "Modules/ModuleManager.h"
 #include "IMeshReductionManagerModule.h"
+#include "IGeometryProcessingInterfacesModule.h"
+#include "GeometryProcessingInterfaces/ApproximateActors.h"
 
 #define LOCTEXT_NAMESPACE "HierarchicalSimplificationCustomizations"
 
@@ -53,9 +57,10 @@ void FHierarchicalSimplificationCustomizations::CustomizeChildren( TSharedRef<IP
 	IDetailGroup& MergeGroup = ChildBuilder.AddGroup(NAME_None, FText::FromString("Mesh generation settings"));
 
 	// Retrieve special case properties
-	SimplifyMeshPropertyHandle = PropertyHandles.FindChecked(GET_MEMBER_NAME_CHECKED(FHierarchicalSimplification, bSimplifyMesh));
+	SimplificationMethodPropertyHandle = PropertyHandles.FindChecked(GET_MEMBER_NAME_CHECKED(FHierarchicalSimplification, SimplificationMethod));
 	TSharedPtr< IPropertyHandle > ProxyMeshSettingPropertyHandle = PropertyHandles.FindChecked(GET_MEMBER_NAME_CHECKED(FHierarchicalSimplification, ProxySetting));
 	TSharedPtr< IPropertyHandle > MergeMeshSettingPropertyHandle = PropertyHandles.FindChecked(GET_MEMBER_NAME_CHECKED(FHierarchicalSimplification, MergeSetting));
+	TSharedPtr< IPropertyHandle > ApproximateMeshSettingsPropertyHandle = PropertyHandles.FindChecked(GET_MEMBER_NAME_CHECKED(FHierarchicalSimplification, ApproximateSettings));
 	TSharedPtr< IPropertyHandle > TransitionScreenSizePropertyHandle = PropertyHandles.FindChecked(GET_MEMBER_NAME_CHECKED(FHierarchicalSimplification, TransitionScreenSize));
 	TSharedPtr< IPropertyHandle > OverrideDrawDistancePropertyHandle = PropertyHandles.FindChecked(GET_MEMBER_NAME_CHECKED(FHierarchicalSimplification, OverrideDrawDistance));
 	TSharedPtr< IPropertyHandle > ReusePreviousLevelClustersPropertyHandle = PropertyHandles.FindChecked(GET_MEMBER_NAME_CHECKED(FHierarchicalSimplification, bReusePreviousLevelClusters));
@@ -63,12 +68,31 @@ void FHierarchicalSimplificationCustomizations::CustomizeChildren( TSharedRef<IP
 	for (auto Iter(PropertyHandles.CreateConstIterator()); Iter; ++Iter)
 	{
 		// Handle special property cases (done inside the loop to maintain order according to the struct
-		if (Iter.Value() == SimplifyMeshPropertyHandle)
+		if (Iter.Value() == SimplificationMethodPropertyHandle)
 		{
-			IDetailPropertyRow& SimplifyMeshRow = MergeGroup.AddPropertyRow(SimplifyMeshPropertyHandle.ToSharedRef());
+			IDetailPropertyRow& SimplifyMeshRow = MergeGroup.AddPropertyRow(SimplificationMethodPropertyHandle.ToSharedRef());
 			AddResetToDefaultOverrides(SimplifyMeshRow);
 
-			SimplifyMeshRow.Visibility(TAttribute<EVisibility>(this, &FHierarchicalSimplificationCustomizations::IsSimplifyMeshVisible));
+			static const UEnum* HLODSimplificationMethodEnum = StaticEnum<EHierarchicalSimplificationMethod>();
+	
+			// Determine whether or not there is a mesh merging interface available (SimplygonMeshReduction/SimplygonSwarm)
+			IMeshReductionModule* ReductionModule = FModuleManager::Get().LoadModulePtr<IMeshReductionModule>("MeshReductionInterface");
+			if (ReductionModule == nullptr || ReductionModule->GetMeshMergingInterface() == nullptr)
+			{
+				static FText RestrictReason = LOCTEXT("SimplifySimplificationMethodUnavailable", "Simplify method is not available, MeshReductionInterface module is missing");
+				TSharedPtr<FPropertyRestriction> EnumRestriction = MakeShared<FPropertyRestriction>(RestrictReason);
+				EnumRestriction->AddHiddenValue(HLODSimplificationMethodEnum->GetNameStringByValue((int64)EHierarchicalSimplificationMethod::Simplify));
+				SimplificationMethodPropertyHandle->AddRestriction(EnumRestriction.ToSharedRef());
+			}			
+
+			IGeometryProcessingInterfacesModule* GeometryProcessingInterfacesModule = FModuleManager::Get().LoadModulePtr<IGeometryProcessingInterfacesModule>("GeometryProcessingInterfaces");
+			if (GeometryProcessingInterfacesModule == nullptr || GeometryProcessingInterfacesModule->GetApproximateActorsImplementation() == nullptr)
+			{
+				static FText RestrictReason = LOCTEXT("ApproximateSimplificationMethodUnavailable", "Approximate method is not available, GeometryProcessingInterfaces module is missing");
+				TSharedPtr<FPropertyRestriction> EnumRestriction = MakeShared<FPropertyRestriction>(RestrictReason);
+				EnumRestriction->AddHiddenValue(HLODSimplificationMethodEnum->GetNameStringByValue((int64)EHierarchicalSimplificationMethod::Approximate));
+				SimplificationMethodPropertyHandle->AddRestriction(EnumRestriction.ToSharedRef());
+			}
 		}
 		else if (Iter.Value() == ProxyMeshSettingPropertyHandle)
 		{
@@ -84,6 +108,13 @@ void FHierarchicalSimplificationCustomizations::CustomizeChildren( TSharedRef<IP
 
 			SettingsRow.Visibility(TAttribute<EVisibility>(this, &FHierarchicalSimplificationCustomizations::IsMergeMeshSettingVisible));
 		}
+		else if (Iter.Value() == ApproximateMeshSettingsPropertyHandle)
+		{
+			IDetailPropertyRow& SettingsRow = MergeGroup.AddPropertyRow(Iter.Value().ToSharedRef());
+			AddResetToDefaultOverrides(SettingsRow);
+
+			SettingsRow.Visibility(TAttribute<EVisibility>(this, &FHierarchicalSimplificationCustomizations::IsApproximateMeshSettingVisible));
+		}		
 		else  if (Iter.Value() == TransitionScreenSizePropertyHandle)
 		{
 			IDetailPropertyRow& SettingsRow = MergeGroup.AddPropertyRow(Iter.Value().ToSharedRef());
@@ -114,41 +145,29 @@ void FHierarchicalSimplificationCustomizations::CustomizeChildren( TSharedRef<IP
 	}
 }
 
-EVisibility FHierarchicalSimplificationCustomizations::IsSimplifyMeshVisible() const
+EHierarchicalSimplificationMethod FHierarchicalSimplificationCustomizations::GetSelectedSimplificationMethod() const
 {
-	// Determine whether or not there is a mesh merging interface available (SimplygonMeshReduction/SimplygonSwarm)
-	IMeshReductionModule& ReductionModule = FModuleManager::Get().LoadModuleChecked<IMeshReductionModule>("MeshReductionInterface");
-	if (ReductionModule.GetMeshMergingInterface() != nullptr)
+	uint8 SimplificationMethod = 0;
+	if (SimplificationMethodPropertyHandle)
 	{
-		return EVisibility::Visible;
+		SimplificationMethodPropertyHandle->GetValue(SimplificationMethod);
 	}
-
-	return EVisibility::Hidden;
+	return (EHierarchicalSimplificationMethod)SimplificationMethod;
 }
 
 EVisibility FHierarchicalSimplificationCustomizations::IsProxyMeshSettingVisible() const
 {
-	bool bSimplifyMesh;
-
-	if (SimplifyMeshPropertyHandle->GetValue(bSimplifyMesh) == FPropertyAccess::Result::Success)
-	{
-		if (IsSimplifyMeshVisible() == EVisibility::Visible && bSimplifyMesh)
-		{
-			return EVisibility::Visible;
-		}
-	}
-		
-	return EVisibility::Hidden;
+	return GetSelectedSimplificationMethod() == EHierarchicalSimplificationMethod::Simplify ? EVisibility::Visible : EVisibility::Hidden;
 }
 
 EVisibility FHierarchicalSimplificationCustomizations::IsMergeMeshSettingVisible() const
 {
-	if(IsProxyMeshSettingVisible() == EVisibility::Hidden)
-	{
-		return EVisibility::Visible;
-	}
+	return GetSelectedSimplificationMethod() == EHierarchicalSimplificationMethod::Merge ? EVisibility::Visible : EVisibility::Hidden;
+}
 
-	return EVisibility::Hidden;
+EVisibility FHierarchicalSimplificationCustomizations::IsApproximateMeshSettingVisible() const
+{
+	return GetSelectedSimplificationMethod() == EHierarchicalSimplificationMethod::Approximate ? EVisibility::Visible : EVisibility::Hidden;
 }
 
 #undef LOCTEXT_NAMESPACE
diff --git a/Engine/Source/Editor/DetailCustomizations/Private/HierarchicalSimplificationCustomizations.h b/Engine/Source/Editor/DetailCustomizations/Private/HierarchicalSimplificationCustomizations.h
--- a/Engine/Source/Editor/DetailCustomizations/Private/HierarchicalSimplificationCustomizations.h
+++ b/Engine/Source/Editor/DetailCustomizations/Private/HierarchicalSimplificationCustomizations.h
@@ -21,9 +21,11 @@ public:
 	virtual void CustomizeChildren( TSharedRef<IPropertyHandle> StructPropertyHandle, class IDetailChildrenBuilder& ChildBuilder, IPropertyTypeCustomizationUtils& StructCustomizationUtils ) override;
 
 protected:
-	EVisibility IsSimplifyMeshVisible() const;
+	EHierarchicalSimplificationMethod GetSelectedSimplificationMethod() const;
+
 	EVisibility IsProxyMeshSettingVisible() const;
 	EVisibility IsMergeMeshSettingVisible() const;
+	EVisibility IsApproximateMeshSettingVisible() const;
 
-	TSharedPtr< IPropertyHandle > SimplifyMeshPropertyHandle;
+	TSharedPtr< IPropertyHandle > SimplificationMethodPropertyHandle;
 };
diff --git a/Engine/Source/Editor/HierarchicalLODOutliner/Private/HierarchicalLODOutlinerModule.cpp b/Engine/Source/Editor/HierarchicalLODOutliner/Private/HierarchicalLODOutlinerModule.cpp
--- a/Engine/Source/Editor/HierarchicalLODOutliner/Private/HierarchicalLODOutlinerModule.cpp
+++ b/Engine/Source/Editor/HierarchicalLODOutliner/Private/HierarchicalLODOutlinerModule.cpp
@@ -58,9 +58,10 @@ void FHierarchicalLODOutlinerModule::OnHLODLevelsArrayChangedEvent()
 			FHierarchicalSimplification& NewLevelSetup = HierarchicalLODSetup.Last();
 			const FHierarchicalSimplification& OldLastLevelSetup = HierarchicalLODSetup[HierarchicalLODSetup.Num() - 2];
 
-			NewLevelSetup.bSimplifyMesh = OldLastLevelSetup.bSimplifyMesh;
+			NewLevelSetup.SimplificationMethod = OldLastLevelSetup.SimplificationMethod;
 			NewLevelSetup.MergeSetting = OldLastLevelSetup.MergeSetting;
 			NewLevelSetup.ProxySetting = OldLastLevelSetup.ProxySetting;
+			NewLevelSetup.ApproximateSettings = OldLastLevelSetup.ApproximateSettings;
 
 			NewLevelSetup.DesiredBoundRadius = OldLastLevelSetup.DesiredBoundRadius * 2.5f;
 			NewLevelSetup.DesiredFillingPercentage = FMath::Max(OldLastLevelSetup.DesiredFillingPercentage * 0.75f, 1.0f);
diff --git a/Engine/Source/Editor/UnrealEd/Private/HierarchicalLOD.cpp b/Engine/Source/Editor/UnrealEd/Private/HierarchicalLOD.cpp
--- a/Engine/Source/Editor/UnrealEd/Private/HierarchicalLOD.cpp
+++ b/Engine/Source/Editor/UnrealEd/Private/HierarchicalLOD.cpp
@@ -925,7 +925,7 @@ void FHierarchicalLODBuilder::BuildMeshesForLODActors(bool bForceAll)
 
 							FMeshDescription* SMDesc = StaticMesh->GetMeshDescription(0);
 
-							if (LODLevelSettings.bSimplifyMesh || LODLevelSettings.MergeSetting.bPivotPointAtZero)
+							if (LODLevelSettings.SimplificationMethod != EHierarchicalSimplificationMethod::Merge || LODLevelSettings.MergeSetting.bPivotPointAtZero)
 							{
 								LODActor->SetActorTransform(FTransform::Identity);
 								FStaticMeshOperations::ApplyTransform(*SMDesc, TransformInv);
diff --git a/Engine/Source/Runtime/Core/Public/UObject/FortniteMainBranchObjectVersion.h b/Engine/Source/Runtime/Core/Public/UObject/FortniteMainBranchObjectVersion.h
--- a/Engine/Source/Runtime/Core/Public/UObject/FortniteMainBranchObjectVersion.h
+++ b/Engine/Source/Runtime/Core/Public/UObject/FortniteMainBranchObjectVersion.h
@@ -238,6 +238,9 @@ struct CORE_API FFortniteMainBranchObjectVersion
 		// Added water non-tessellated LOD
 		WaterNontessellatedLODSupportAdded,
 
+		// Added FHierarchicalSimplification::SimplificationMethod
+		HierarchicalSimplificationMethodEnumAdded,
+
 		// -----<new versions can be added above this line>-------------------------------------------------
 		VersionPlusOne,
 		LatestVersion = VersionPlusOne - 1
diff --git a/Engine/Source/Runtime/Engine/Classes/Engine/Level.h b/Engine/Source/Runtime/Engine/Classes/Engine/Level.h
--- a/Engine/Source/Runtime/Engine/Classes/Engine/Level.h
+++ b/Engine/Source/Runtime/Engine/Classes/Engine/Level.h
@@ -400,33 +400,6 @@ struct ENGINE_API FLevelSimplificationDetails
 	UPROPERTY(Category=Landscape, EditAnywhere)
 	bool bBakeGrassToLandscape;
 
-	UPROPERTY()
-	bool bGenerateMeshNormalMap_DEPRECATED;
-	
-	UPROPERTY()
-	bool bGenerateMeshMetallicMap_DEPRECATED;
-
-	UPROPERTY()
-	bool bGenerateMeshRoughnessMap_DEPRECATED;
-	
-	UPROPERTY()
-	bool bGenerateMeshSpecularMap_DEPRECATED;
-	
-	UPROPERTY()
-	bool bGenerateLandscapeNormalMap_DEPRECATED;
-
-	UPROPERTY()
-	bool bGenerateLandscapeMetallicMap_DEPRECATED;
-
-	UPROPERTY()
-	bool bGenerateLandscapeRoughnessMap_DEPRECATED;
-	
-	UPROPERTY()
-	bool bGenerateLandscapeSpecularMap_DEPRECATED;
-
-	/** Handles deprecated properties */
-	void PostLoadDeprecated();
-	
 	FLevelSimplificationDetails();
 	bool operator == (const FLevelSimplificationDetails& Other) const;
 };
diff --git a/Engine/Source/Runtime/Engine/Classes/Engine/MeshMerging.h b/Engine/Source/Runtime/Engine/Classes/Engine/MeshMerging.h
--- a/Engine/Source/Runtime/Engine/Classes/Engine/MeshMerging.h
+++ b/Engine/Source/Runtime/Engine/Classes/Engine/MeshMerging.h
@@ -409,7 +409,18 @@ struct FMeshProxySettings
 
 #if WITH_EDITORONLY_DATA
 	/** Handles deprecated properties */
-	void PostLoadDeprecated();
+	void PostSerialize(const FArchive& Ar);
+#endif
+};
+
+template<>
+struct TStructOpsTypeTraits<FMeshProxySettings> : public TStructOpsTypeTraitsBase2<FMeshProxySettings>
+{
+#if WITH_EDITORONLY_DATA
+	enum
+	{
+		WithPostSerialize = true,
+	};
 #endif
 };
 
@@ -466,14 +477,14 @@ struct FMeshMergingSettings
 	UPROPERTY(EditAnywhere, Category = MaterialSettings, meta=(DisplayAfter="MaterialSettings"))
 	int32 GutterSize;
 
-	// A given LOD level to export from the source meshes
-	UPROPERTY(EditAnywhere, Category = MeshSettings, BlueprintReadWrite, meta = (DisplayAfter="LODSelectionType", ClampMin = "0", ClampMax = "7", UIMin = "0", UIMax = "7", EnumCondition = 1))
-	int32 SpecificLOD;
-
 	/** Which selection mode should be used when generating the merged static mesh */
-	UPROPERTY(EditAnywhere, BlueprintReadWrite, Category = MeshSettings, meta = (DisplayAfter="bBakeVertexDataToMesh"))
+	UPROPERTY(EditAnywhere, BlueprintReadWrite, Category = MeshSettings, meta = (DisplayAfter="bBakeVertexDataToMesh", DisplayName = "LOD Selection Type"))
 	EMeshLODSelectionType LODSelectionType;
 
+	/** A given LOD level to export from the source meshes, used if LOD Selection Type is set to SpecificLOD */
+	UPROPERTY(EditAnywhere, Category = MeshSettings, BlueprintReadWrite, meta = (DisplayAfter="LODSelectionType", EditCondition = "LODSelectionType == EMeshLODSelectionType::SpecificLOD", ClampMin = "0", ClampMax = "7", UIMin = "0", UIMax = "7", EnumCondition = 1))
+	int32 SpecificLOD;
+
 	/** Whether to generate lightmap UVs for a merged mesh*/
 	UPROPERTY(EditAnywhere, AdvancedDisplay, BlueprintReadWrite, Category = MeshSettings, meta=(DisplayName="Generate Lightmap UV"))
 	uint8 bGenerateLightMapUV:1;
@@ -490,12 +501,12 @@ struct FMeshMergingSettings
 	UPROPERTY(EditAnywhere, BlueprintReadWrite, Category = MeshSettings)
 	uint8 bMergePhysicsData:1;
 
-	/** Whether to merge source materials into one flat material, ONLY available when merging a single LOD level, see LODSelectionType */
-	UPROPERTY(EditAnywhere, BlueprintReadWrite, Category = MaterialSettings)
+	/** Whether to merge source materials into one flat material, ONLY available when LOD Selection Type is set to LowestDetailLOD */
+	UPROPERTY(EditAnywhere, BlueprintReadWrite, Category = MaterialSettings, meta=(EditCondition="LODSelectionType == EMeshLODSelectionType::LowestDetailLOD"))
 	uint8 bMergeMaterials:1;
 
 	/** Create a flat material from all source materials, along with a new set of UVs. This material won't be applied to any section by default. */
-	UPROPERTY(EditAnywhere, BlueprintReadWrite, Category = MaterialSettings)
+	UPROPERTY(EditAnywhere, AdvancedDisplay, BlueprintReadWrite, Category = MaterialSettings, meta=(EditCondition="LODSelectionType == EMeshLODSelectionType::LowestDetailLOD"))
 	uint8 bCreateMergedMaterial : 1;
 
 	/** Whether or not vertex data such as vertex colours should be baked into the resulting mesh */
@@ -506,7 +517,7 @@ struct FMeshMergingSettings
 	UPROPERTY(EditAnywhere, Category = MaterialSettings, BlueprintReadWrite, meta = (EditCondition = "bMergeMaterials"))
 	uint8 bUseVertexDataForBakingMaterial:1;
 
-	// Whether or not to calculate varying output texture sizes according to their importance in the final atlas texture
+	/** Whether or not to calculate varying output texture sizes according to their importance in the final atlas texture */
 	UPROPERTY(Category = MaterialSettings, EditAnywhere, BlueprintReadWrite, meta = (EditCondition = "bMergeMaterials"))
 	uint8 bUseTextureBinning:1;
 
@@ -569,8 +580,8 @@ struct FMeshMergingSettings
 	FMeshMergingSettings()
 		: TargetLightMapResolution(256)
 		, GutterSize(2)
-		, SpecificLOD(0)
 		, LODSelectionType(EMeshLODSelectionType::CalculateLOD)
+		, SpecificLOD(0)
 		, bGenerateLightMapUV(true)
 		, bComputedLightMapResolution(false)
 		, bPivotPointAtZero(false)
@@ -607,7 +618,18 @@ struct FMeshMergingSettings
 
 #if WITH_EDITORONLY_DATA
 	/** Handles deprecated properties */
-	void PostLoadDeprecated();
+	void PostSerialize(const FArchive& Ar);
+#endif
+};
+
+template<>
+struct TStructOpsTypeTraits<FMeshMergingSettings> : public TStructOpsTypeTraitsBase2<FMeshMergingSettings>
+{
+#if WITH_EDITORONLY_DATA
+	enum
+	{
+		WithPostSerialize = true,
+	};
 #endif
 };
 
@@ -965,9 +987,4 @@ struct FMeshApproximationSettings
 	{
 		return !(*this == Other);
 	}
-
-#if WITH_EDITORONLY_DATA
-	/** Handles deprecated properties */
-	void PostLoadDeprecated() {}		// none currently
-#endif
 };
\ No newline at end of file
diff --git a/Engine/Source/Runtime/Engine/Classes/GameFramework/WorldSettings.h b/Engine/Source/Runtime/Engine/Classes/GameFramework/WorldSettings.h
--- a/Engine/Source/Runtime/Engine/Classes/GameFramework/WorldSettings.h
+++ b/Engine/Source/Runtime/Engine/Classes/GameFramework/WorldSettings.h
@@ -267,6 +267,16 @@ struct ENGINE_API FNetViewer
 	FNetViewer(AController* InController);
 };
 
+UENUM()
+enum class EHierarchicalSimplificationMethod : uint8
+{
+	None = 0			UMETA(hidden),
+	Merge = 1,
+	Simplify = 2,
+	Approximate = 3
+};
+
+
 USTRUCT()
 struct ENGINE_API FHierarchicalSimplification
 {
@@ -285,14 +295,6 @@ struct ENGINE_API FHierarchicalSimplification
 	UPROPERTY(Category = FHierarchicalSimplification, EditAnywhere, AdvancedDisplay)
 	uint8 bAllowSpecificExclusion : 1;
 
-	/** If this is true, it will simplify mesh but it is slower.
-	* If false, it will just merge actors but not simplify using the lower LOD if exists.
-	* For example if you build LOD 1, it will use LOD 1 of the mesh to merge actors if exists.
-	* If you merge material, it will reduce drawcalls.
-	*/
-	UPROPERTY(Category = FHierarchicalSimplification, EditAnywhere)
-	uint8 bSimplifyMesh:1;
-
 	/** Only generate clusters for HLOD volumes */
 	UPROPERTY(EditAnywhere, Category = FHierarchicalSimplification, AdvancedDisplay, meta = (editcondition = "!bReusePreviousLevelClusters", DisplayAfter="MinNumberOfActorsToBuild"))
 	uint8 bOnlyGenerateClustersForVolumes:1;
@@ -301,14 +303,21 @@ struct ENGINE_API FHierarchicalSimplification
 	UPROPERTY(EditAnywhere, Category = FHierarchicalSimplification, AdvancedDisplay, meta=(DisplayAfter="bOnlyGenerateClustersForVolumes"))
 	uint8 bReusePreviousLevelClusters:1;
 
-	/** Simplification Setting if bSimplifyMesh is true */
+	UPROPERTY(Category = FHierarchicalSimplification, EditAnywhere)
+	EHierarchicalSimplificationMethod SimplificationMethod;
+
+	/** Simplification settings, used if SimplificationMethod is Simplify */
 	UPROPERTY(Category = FHierarchicalSimplification, EditAnywhere, AdvancedDisplay)
 	FMeshProxySettings ProxySetting;
 
-	/** Merge Mesh Setting if bSimplifyMesh is false */
+	/** Merge settings, used if SimplificationMethod is Merge */
 	UPROPERTY(Category = FHierarchicalSimplification, EditAnywhere, AdvancedDisplay)
 	FMeshMergingSettings MergeSetting;
 
+	/** Approximate settings, used if SimplificationMethod is Approximate */
+	UPROPERTY(Category = FHierarchicalSimplification, EditAnywhere, AdvancedDisplay)
+	FMeshApproximationSettings ApproximateSettings;
+
 	/** Desired Bounding Radius for clustering - this is not guaranteed but used to calculate filling factor for auto clustering */
 	UPROPERTY(EditAnywhere, Category=FHierarchicalSimplification, AdvancedDisplay, meta=(UIMin=10.f, ClampMin=10.f, editcondition = "!bReusePreviousLevelClusters"))
 	float DesiredBoundRadius;
@@ -319,16 +328,21 @@ struct ENGINE_API FHierarchicalSimplification
 
 	/** Min number of actors to build LODActor */
 	UPROPERTY(EditAnywhere, Category=FHierarchicalSimplification, AdvancedDisplay, meta=(ClampMin = "1", UIMin = "1", editcondition = "!bReusePreviousLevelClusters"))
-	int32 MinNumberOfActorsToBuild;	
+	int32 MinNumberOfActorsToBuild;
+
+#if WITH_EDITORONLY_DATA
+	UPROPERTY(meta = (DeprecatedProperty))
+	uint8 bSimplifyMesh_DEPRECATED:1;
+#endif
 
 	FHierarchicalSimplification()
 		: TransitionScreenSize(0.315f)
 		, OverrideDrawDistance(10000)
 		, bUseOverrideDrawDistance(false)
 		, bAllowSpecificExclusion(false)
-		, bSimplifyMesh(false)
 		, bOnlyGenerateClustersForVolumes(false)
 		, bReusePreviousLevelClusters(false)
+		, SimplificationMethod(EHierarchicalSimplificationMethod::Merge)
 		, DesiredBoundRadius(2000)
 		, DesiredFillingPercentage(50)
 		, MinNumberOfActorsToBuild(2)
@@ -338,6 +352,25 @@ struct ENGINE_API FHierarchicalSimplification
 		ProxySetting.MaterialSettings.MaterialMergeType = EMaterialMergeType::MaterialMergeType_Simplygon;
 		ProxySetting.bCreateCollision = false;
 	}
+
+#if WITH_EDITORONLY_DATA
+	bool Serialize(FArchive& Ar);
+
+	/** Handles deprecated properties */
+	void PostSerialize(const FArchive& Ar);
+#endif
+};
+
+template<>
+struct TStructOpsTypeTraits<FHierarchicalSimplification> : public TStructOpsTypeTraitsBase2<FHierarchicalSimplification>
+{
+#if WITH_EDITORONLY_DATA
+	enum
+	{
+		WithSerializer = true,
+		WithPostSerialize = true,
+	};
+#endif
 };
 
 UCLASS(Blueprintable)
diff --git a/Engine/Source/Runtime/Engine/Private/EngineTypes.cpp b/Engine/Source/Runtime/Engine/Private/EngineTypes.cpp
--- a/Engine/Source/Runtime/Engine/Private/EngineTypes.cpp
+++ b/Engine/Source/Runtime/Engine/Private/EngineTypes.cpp
@@ -23,50 +23,56 @@ FDetachmentTransformRules FDetachmentTransformRules::KeepRelativeTransform(EDeta
 FDetachmentTransformRules FDetachmentTransformRules::KeepWorldTransform(EDetachmentRule::KeepWorld, true);
 
 #if WITH_EDITORONLY_DATA
-void FMeshProxySettings::PostLoadDeprecated()
+void FMeshProxySettings::PostSerialize(const FArchive& Ar)
 {
-	MaterialSettings.MaterialMergeType = EMaterialMergeType::MaterialMergeType_Simplygon;
+	if (Ar.IsLoading())
+	{
+		MaterialSettings.MaterialMergeType = EMaterialMergeType::MaterialMergeType_Simplygon;
+	}
 }
 
-void FMeshMergingSettings::PostLoadDeprecated()
+void FMeshMergingSettings::PostSerialize(const FArchive& Ar)
 {
-	FMeshMergingSettings DefaultObject;
-	if (bImportVertexColors_DEPRECATED != DefaultObject.bImportVertexColors_DEPRECATED)
+	if (Ar.IsLoading())
 	{
-		bBakeVertexDataToMesh = bImportVertexColors_DEPRECATED;
-	}
+		FMeshMergingSettings DefaultObject;
+		if (bImportVertexColors_DEPRECATED != DefaultObject.bImportVertexColors_DEPRECATED)
+		{
+			bBakeVertexDataToMesh = bImportVertexColors_DEPRECATED;
+		}
 
-	if (bExportNormalMap_DEPRECATED != DefaultObject.bExportNormalMap_DEPRECATED)
-	{
-		MaterialSettings.bNormalMap = bExportNormalMap_DEPRECATED;
-	}
+		if (bExportNormalMap_DEPRECATED != DefaultObject.bExportNormalMap_DEPRECATED)
+		{
+			MaterialSettings.bNormalMap = bExportNormalMap_DEPRECATED;
+		}
 
-	if (bExportMetallicMap_DEPRECATED != DefaultObject.bExportMetallicMap_DEPRECATED)
-	{
-		MaterialSettings.bMetallicMap = bExportMetallicMap_DEPRECATED;
-	}
-	if (bExportRoughnessMap_DEPRECATED != DefaultObject.bExportRoughnessMap_DEPRECATED)
-	{
-		MaterialSettings.bRoughnessMap = bExportRoughnessMap_DEPRECATED;
-	}
-	if (bExportSpecularMap_DEPRECATED != DefaultObject.bExportSpecularMap_DEPRECATED)
-	{
-		MaterialSettings.bSpecularMap = bExportSpecularMap_DEPRECATED;
-	}
-	if (MergedMaterialAtlasResolution_DEPRECATED != DefaultObject.MergedMaterialAtlasResolution_DEPRECATED)
-	{
-		MaterialSettings.TextureSize.X = MergedMaterialAtlasResolution_DEPRECATED;
-		MaterialSettings.TextureSize.Y = MergedMaterialAtlasResolution_DEPRECATED;
-	}
-	if (bCalculateCorrectLODModel_DEPRECATED != DefaultObject.bCalculateCorrectLODModel_DEPRECATED)
-	{
-		LODSelectionType = EMeshLODSelectionType::CalculateLOD;
-	}
+		if (bExportMetallicMap_DEPRECATED != DefaultObject.bExportMetallicMap_DEPRECATED)
+		{
+			MaterialSettings.bMetallicMap = bExportMetallicMap_DEPRECATED;
+		}
+		if (bExportRoughnessMap_DEPRECATED != DefaultObject.bExportRoughnessMap_DEPRECATED)
+		{
+			MaterialSettings.bRoughnessMap = bExportRoughnessMap_DEPRECATED;
+		}
+		if (bExportSpecularMap_DEPRECATED != DefaultObject.bExportSpecularMap_DEPRECATED)
+		{
+			MaterialSettings.bSpecularMap = bExportSpecularMap_DEPRECATED;
+		}
+		if (MergedMaterialAtlasResolution_DEPRECATED != DefaultObject.MergedMaterialAtlasResolution_DEPRECATED)
+		{
+			MaterialSettings.TextureSize.X = MergedMaterialAtlasResolution_DEPRECATED;
+			MaterialSettings.TextureSize.Y = MergedMaterialAtlasResolution_DEPRECATED;
+		}
+		if (bCalculateCorrectLODModel_DEPRECATED != DefaultObject.bCalculateCorrectLODModel_DEPRECATED)
+		{
+			LODSelectionType = EMeshLODSelectionType::CalculateLOD;
+		}
 
-	if (ExportSpecificLOD_DEPRECATED != DefaultObject.ExportSpecificLOD_DEPRECATED)
-	{
-		SpecificLOD = ExportSpecificLOD_DEPRECATED;
-		LODSelectionType = EMeshLODSelectionType::SpecificLOD;
+		if (ExportSpecificLOD_DEPRECATED != DefaultObject.ExportSpecificLOD_DEPRECATED)
+		{
+			SpecificLOD = ExportSpecificLOD_DEPRECATED;
+			LODSelectionType = EMeshLODSelectionType::SpecificLOD;
+		}
 	}
 }
 #endif
diff --git a/Engine/Source/Runtime/Engine/Private/HLODProxy.cpp b/Engine/Source/Runtime/Engine/Private/HLODProxy.cpp
--- a/Engine/Source/Runtime/Engine/Private/HLODProxy.cpp
+++ b/Engine/Source/Runtime/Engine/Private/HLODProxy.cpp
@@ -510,7 +510,17 @@ FName UHLODProxy::GenerateKeyForActor(const ALODActor* LODActor, bool bMustUndoL
 	if (HierarchicalLODSetups.IsValidIndex(LODActor->LODLevel - 1))
 	{
 		const FHierarchicalSimplification& HLODSettings = HierarchicalLODSetups[LODActor->LODLevel - 1];
-		bConsiderPhysicData = HLODSettings.bSimplifyMesh ? HLODSettings.ProxySetting.bCreateCollision : HLODSettings.MergeSetting.bMergePhysicsData;
+		switch (HLODSettings.SimplificationMethod)
+		{
+		case EHierarchicalSimplificationMethod::Simplify:
+			bConsiderPhysicData = HLODSettings.ProxySetting.bCreateCollision;
+			break;
+		case EHierarchicalSimplificationMethod::Merge:
+			bConsiderPhysicData = HLODSettings.MergeSetting.bMergePhysicsData;
+			break;
+		default:
+			bConsiderPhysicData = false;
+		}
 	}
 	
 	// Base us off the unique object ID
diff --git a/Engine/Source/Runtime/Engine/Private/Level.cpp b/Engine/Source/Runtime/Engine/Private/Level.cpp
--- a/Engine/Source/Runtime/Engine/Private/Level.cpp
+++ b/Engine/Source/Runtime/Engine/Private/Level.cpp
@@ -401,14 +401,6 @@ FLevelSimplificationDetails::FLevelSimplificationDetails()
  , LandscapeExportLOD(7)
  , bBakeFoliageToLandscape(false)
  , bBakeGrassToLandscape(false)
- , bGenerateMeshNormalMap_DEPRECATED(true)
- , bGenerateMeshMetallicMap_DEPRECATED(false)
- , bGenerateMeshRoughnessMap_DEPRECATED(false)
- , bGenerateMeshSpecularMap_DEPRECATED(false)
- , bGenerateLandscapeNormalMap_DEPRECATED(true)
- , bGenerateLandscapeMetallicMap_DEPRECATED(false)
- , bGenerateLandscapeRoughnessMap_DEPRECATED(false)
- , bGenerateLandscapeSpecularMap_DEPRECATED(false)
 {
 }
 
@@ -424,12 +416,6 @@ bool FLevelSimplificationDetails::operator == (const FLevelSimplificationDetails
 		&& bBakeGrassToLandscape == Other.bBakeGrassToLandscape;
 }
 
-void FLevelSimplificationDetails::PostLoadDeprecated()
-{
-	FLevelSimplificationDetails DefaultObject;
-
-}
-
 static bool IsActorFolderObjectsFeatureAvailable()
 {
 #if WITH_EDITOR
@@ -1235,12 +1221,6 @@ void ULevel::PostLoad()
 		}
 	}
 
-	// Fixup deprecated stuff in levels simplification settings
-	for (int32 Index = 0; Index < UE_ARRAY_COUNT(LevelSimplification); ++Index)
-	{
-		LevelSimplification[Index].PostLoadDeprecated();
-	}
-
 	RepairLevelScript();
 #endif
 }
diff --git a/Engine/Source/Runtime/Engine/Private/WorldSettings.cpp b/Engine/Source/Runtime/Engine/Private/WorldSettings.cpp
--- a/Engine/Source/Runtime/Engine/Private/WorldSettings.cpp
+++ b/Engine/Source/Runtime/Engine/Private/WorldSettings.cpp
@@ -25,6 +25,7 @@
 #include "PhysicsEngine/PhysicsSettings.h"
 #include "UObject/ReleaseObjectVersion.h"
 #include "UObject/EnterpriseObjectVersion.h"
+#include "UObject/FortniteMainBranchObjectVersion.h"
 #include "SceneManagement.h"
 #include "AI/AISystemBase.h"
 #include "AI/NavigationSystemConfig.h"
@@ -557,12 +558,6 @@ void AWorldSettings::PostLoad()
 	Super::PostLoad();
 
 #if WITH_EDITOR
-	for (FHierarchicalSimplification& Entry : HierarchicalLODSetup)
-	{
-		Entry.ProxySetting.PostLoadDeprecated();
-		Entry.MergeSetting.PostLoadDeprecated();
-	}
-
 	if (WorldPartition)
 	{
 		// Force to re-apply WorldPartition restrictions on WorldSettings (in case they changed)
@@ -1036,4 +1031,27 @@ void AWorldSettings::RewindForReplay()
 	bHighPriorityLoading = false;
 }
 
+#if WITH_EDITORONLY_DATA
+
+bool FHierarchicalSimplification::Serialize(FArchive& Ar)
+{
+	Ar.UsingCustomVersion(FFortniteMainBranchObjectVersion::GUID);
+
+	// Don't actually serialize, just write the custom version for PostSerialize
+	return false;
+}
+
+void FHierarchicalSimplification::PostSerialize(const FArchive& Ar)
+{
+	if (Ar.IsLoading())
+	{
+		if (Ar.CustomVer(FFortniteMainBranchObjectVersion::GUID) < FFortniteMainBranchObjectVersion::HierarchicalSimplificationMethodEnumAdded)
+		{
+			SimplificationMethod = bSimplifyMesh_DEPRECATED ? EHierarchicalSimplificationMethod::Simplify : EHierarchicalSimplificationMethod::Merge;
+		}
+	}
+}
+
+#endif
+
 #undef LOCTEXT_NAMESPACE
