Fixed lost HLOD assets names (only prefix would remain) in previous BuildStaticMeshForLODActor() refactor

diff --git a/Engine/Source/Developer/HierarchicalLODUtilities/Private/HierarchicalLODUtilities.cpp b/Engine/Source/Developer/HierarchicalLODUtilities/Private/HierarchicalLODUtilities.cpp
--- a/Engine/Source/Developer/HierarchicalLODUtilities/Private/HierarchicalLODUtilities.cpp
+++ b/Engine/Source/Developer/HierarchicalLODUtilities/Private/HierarchicalLODUtilities.cpp
@@ -482,7 +482,7 @@ FHLODBuildResults GenerateHLODMesh_Simplify(const FHLODBuildParams& InBuildParam
 	TArray<UStaticMeshComponent*> StaticMeshComponents;
 	Algo::Transform(InBuildParams.Components, StaticMeshComponents, [](UPrimitiveComponent* InPrimitiveComponent) { return Cast<UStaticMeshComponent>(InPrimitiveComponent); });
 
-	MeshMergeUtilities.CreateProxyMesh(StaticMeshComponents, ProxySettings, InBuildParams.BaseMaterial, InBuildParams.AssetsOuter, "", JobID, Processor->GetCallbackDelegate(), true, OverrideLODSetup.TransitionScreenSize);
+	MeshMergeUtilities.CreateProxyMesh(StaticMeshComponents, ProxySettings, InBuildParams.BaseMaterial, InBuildParams.AssetsOuter, InBuildParams.PackageName, JobID, Processor->GetCallbackDelegate(), true, OverrideLODSetup.TransitionScreenSize);
 
 	
 	HLODBuildResults.bDeferredResults = true;
@@ -506,7 +506,7 @@ FHLODBuildResults GenerateHLODMesh_Merge(const FHLODBuildParams& InBuildParams)
 	FHLODBuildResults HLODBuildResults;
 
 	TArray<UObject*> OutAssets;
-	MeshMergeUtilities.MergeComponentsToStaticMesh(InBuildParams.Components, InBuildParams.LODActor->GetWorld(), MergeSettings, InBuildParams.BaseMaterial, InBuildParams.AssetsOuter, "", OutAssets, HLODBuildResults.HLODLocation, InBuildParams.LODSetup.TransitionScreenSize, true);
+	MeshMergeUtilities.MergeComponentsToStaticMesh(InBuildParams.Components, InBuildParams.LODActor->GetWorld(), MergeSettings, InBuildParams.BaseMaterial, InBuildParams.AssetsOuter, InBuildParams.PackageName, OutAssets, HLODBuildResults.HLODLocation, InBuildParams.LODSetup.TransitionScreenSize, true);
 
 	// set staticmesh
 	for (UObject* Asset : OutAssets)
@@ -745,6 +745,15 @@ bool FHierarchicalLODUtilities::BuildStaticMeshForLODActor(ALODActor* LODActor,
 		HLODBuildParams.BaseMaterial = InBaseMaterial;
 		HLODBuildParams.AssetsOuter = AssetsOuter;
 
+		// Should give a unique name, so use the LODActor tag, or if empty, the first actor name
+		FString LODActorTag = LODActor->GetLODActorTag();
+		if (LODActorTag.IsEmpty())
+		{
+			const AActor* FirstActor = UHLODProxy::FindFirstActor(LODActor);
+			LODActorTag = *FirstActor->GetName();
+		}
+		HLODBuildParams.PackageName = FString::Printf(TEXT("%s_%i_%s"), *(AssetsOuter->GetName()), LODActor->LODLevel - 1, *LODActorTag);
+
 		FHLODBuildResults HLODBuildResults;
 		switch (LODSetup.SimplificationMethod)
 		{
