Use PerPlatformProperties properly when in preview

diff --git a/Engine/Source/Developer/HierarchicalLODUtilities/Private/HierarchicalLODUtilities.cpp b/Engine/Source/Developer/HierarchicalLODUtilities/Private/HierarchicalLODUtilities.cpp
--- a/Engine/Source/Developer/HierarchicalLODUtilities/Private/HierarchicalLODUtilities.cpp
+++ b/Engine/Source/Developer/HierarchicalLODUtilities/Private/HierarchicalLODUtilities.cpp
@@ -1315,6 +1315,7 @@ int32 FHierarchicalLODUtilities::ComputeStaticMeshLODLevel(const TArray<FStaticM
 	// Walk backwards and return the first matching LOD
 	for (int32 LODIndex = NumLODs - 1; LODIndex >= 0; --LODIndex)
 	{
+		// For HLOD generation we want the default values and not the per platform overrides
 		if (SourceModels[LODIndex].ScreenSize.Default > ScreenSize || ((SourceModels[LODIndex].ScreenSize.Default == 0.0f) && (RenderData->ScreenSize[LODIndex].Default != SourceModels[LODIndex].ScreenSize.Default) && (RenderData->ScreenSize[LODIndex].Default > ScreenSize)))
 		{
 			return FMath::Max(LODIndex, 0);
diff --git a/Engine/Source/Editor/StaticMeshEditor/Private/StaticMeshEditorSubsystem.cpp b/Engine/Source/Editor/StaticMeshEditor/Private/StaticMeshEditorSubsystem.cpp
--- a/Engine/Source/Editor/StaticMeshEditor/Private/StaticMeshEditorSubsystem.cpp
+++ b/Engine/Source/Editor/StaticMeshEditor/Private/StaticMeshEditorSubsystem.cpp
@@ -992,7 +992,7 @@ TArray<float> UStaticMeshEditorSubsystem::GetLodScreenSizes(UStaticMesh* StaticM
 	{
 		if (StaticMesh->GetRenderData())
 		{
-			float CurScreenSize = StaticMesh->GetRenderData()->ScreenSize[i].Default;
+			float CurScreenSize = StaticMesh->GetRenderData()->ScreenSize[i].GetValue();
 			ScreenSizes.Add(CurScreenSize);
 		}
 		else
diff --git a/Engine/Source/Editor/StaticMeshEditor/Private/StaticMeshEditorTools.cpp b/Engine/Source/Editor/StaticMeshEditor/Private/StaticMeshEditorTools.cpp
--- a/Engine/Source/Editor/StaticMeshEditor/Private/StaticMeshEditorTools.cpp
+++ b/Engine/Source/Editor/StaticMeshEditor/Private/StaticMeshEditorTools.cpp
@@ -4058,6 +4058,12 @@ float FLevelOfDetailSettingsLayout::GetLODScreenSize(FName PlatformGroupName, in
 	if(Mesh->bAutoComputeLODScreenSize && Mesh->GetRenderData())
 	{
 		ScreenSize = Mesh->GetRenderData()->ScreenSize[LODIndex].Default;
+		const float* PlatformScreenSize = Mesh->GetRenderData()->ScreenSize[LODIndex].PerPlatform.Find(PlatformGroupName);
+		if (PlatformScreenSize != nullptr)
+		{
+			ScreenSize = *PlatformScreenSize;
+		}
+
 	}
 	else if(Mesh->IsSourceModelValid(LODIndex))
 	{
