OnlyEditorOnly support for HLOD implicit package references: Declare the packages that will be loaded if they exist as dependencies.
#rnx

diff --git a/Engine/Source/Developer/HierarchicalLODUtilities/Private/HierarchicalLODUtilities.cpp b/Engine/Source/Developer/HierarchicalLODUtilities/Private/HierarchicalLODUtilities.cpp
--- a/Engine/Source/Developer/HierarchicalLODUtilities/Private/HierarchicalLODUtilities.cpp
+++ b/Engine/Source/Developer/HierarchicalLODUtilities/Private/HierarchicalLODUtilities.cpp
@@ -52,6 +52,7 @@
 #include "LevelUtils.h"
 #include "Materials/Material.h"
 #include "Materials/MaterialInstance.h"
+#include "UObject/ICookInfo.h"
 
 DEFINE_LOG_CATEGORY_STATIC(LogHierarchicalLODUtilities, Verbose, All);
 
@@ -133,6 +134,12 @@ static FString GetHLODPackageName(const ULevel* InLevel, const uint32 InHLODLeve
 	return GetHLODPackageName(LevelPackageName, InHLODLevelIndex, InOutHLODProxyName);
 }
 
+FString FHierarchicalLODUtilities::GetHLODPackageName(const ULevel* InLevel, const uint32 InHLODLevelIndex)
+{
+	FString HLODProxyName;
+	return ::GetHLODPackageName(InLevel, InHLODLevelIndex, HLODProxyName);
+}
+
 void FHierarchicalLODUtilities::CleanStandaloneAssetsInPackage(UPackage* InPackage)
 {
 	TArray<UObject*> Objects;
@@ -181,7 +188,7 @@ UPackage* FHierarchicalLODUtilities::CreateOrRetrieveLevelHLODPackage(const ULev
 	checkf(InLevel != nullptr, TEXT("Invalid Level supplied"));
 
 	FString HLODProxyName;
-	const FString HLODLevelPackageName = GetHLODPackageName(InLevel, HLODLevelIndex, HLODProxyName);
+	const FString HLODLevelPackageName = ::GetHLODPackageName(InLevel, HLODLevelIndex, HLODProxyName);
 
 	// Find existing package
 	bool bCreatedNewPackage = false;
@@ -199,7 +206,8 @@ UHLODProxy* FHierarchicalLODUtilities::RetrieveLevelHLODProxy(const ULevel* InLe
 {
 	checkf(InLevel != nullptr, TEXT("Invalid Level supplied"));
 	FString HLODProxyName;
-	const FString HLODLevelPackageName = GetHLODPackageName(InLevel, HLODLevelIndex, HLODProxyName);
+	FCookLoadScope CookLoadScope(ECookLoadType::UsedInGame);
+	const FString HLODLevelPackageName = ::GetHLODPackageName(InLevel, HLODLevelIndex, HLODProxyName);
 
 	UHLODProxy* HLODProxy = LoadObject<UHLODProxy>(nullptr, *HLODLevelPackageName, nullptr, LOAD_Quiet | LOAD_NoWarn);
 	return HLODProxy;
@@ -240,7 +248,7 @@ UPackage* FHierarchicalLODUtilities::CreateOrRetrieveLevelHLODPackage(const ULev
 FString FHierarchicalLODUtilities::GetLevelHLODProxyName(const FString& InLevelPackageName, const uint32 InHLODLevelIndex)
 {
 	FString HLODProxyName;
-	FString HLODPackageName = GetHLODPackageName(InLevelPackageName, InHLODLevelIndex, HLODProxyName);
+	FString HLODPackageName = ::GetHLODPackageName(InLevelPackageName, InHLODLevelIndex, HLODProxyName);
 	return HLODPackageName + TEXT(".") + HLODProxyName;
 }
 
diff --git a/Engine/Source/Developer/HierarchicalLODUtilities/Public/HierarchicalLODUtilities.h b/Engine/Source/Developer/HierarchicalLODUtilities/Public/HierarchicalLODUtilities.h
--- a/Engine/Source/Developer/HierarchicalLODUtilities/Public/HierarchicalLODUtilities.h
+++ b/Engine/Source/Developer/HierarchicalLODUtilities/Public/HierarchicalLODUtilities.h
@@ -33,6 +33,7 @@ public:
 	virtual UPackage* RetrieveLevelHLODPackage(const ULevel* InLevel, const uint32 HLODLevelIndex) override;
 	virtual UHLODProxy* CreateOrRetrieveLevelHLODProxy(const ULevel* InLevel, const uint32 HLODLevelIndex) override;
 	virtual UHLODProxy* RetrieveLevelHLODProxy(const ULevel* InLevel, const uint32 HLODLevelIndex) override;
+	virtual FString GetHLODPackageName(const ULevel* InLevel, const uint32 InHLODLevelIndex) override;
 	virtual FString GetLevelHLODProxyName(const FString& InLevelPackageName, const uint32 HLODLevelIndex) override;
 	UE_DEPRECATED(4.20, "Use different signature containing BaseMaterial parameter")
 	virtual bool BuildStaticMeshForLODActor(ALODActor* LODActor, UPackage* AssetsOuter, const FHierarchicalSimplification& LODSetup) override;
diff --git a/Engine/Source/Developer/HierarchicalLODUtilities/Public/IHierarchicalLODUtilities.h b/Engine/Source/Developer/HierarchicalLODUtilities/Public/IHierarchicalLODUtilities.h
--- a/Engine/Source/Developer/HierarchicalLODUtilities/Public/IHierarchicalLODUtilities.h
+++ b/Engine/Source/Developer/HierarchicalLODUtilities/Public/IHierarchicalLODUtilities.h
@@ -75,6 +75,8 @@ public:
 	/** Retrieves the HLOD proxy that is created for the given level */
 	virtual UHLODProxy* RetrieveLevelHLODProxy(const ULevel* InLevel, const uint32 HLODLevelIndex) = 0;
 
+	virtual FString GetHLODPackageName(const ULevel* InLevel, const uint32 InHLODLevelIndex) = 0;
+
 	/** Get the name of a HLOD proxy given the name of a level's package */
 	virtual FString GetLevelHLODProxyName(const FString& InLevelPackageName, const uint32 HLODLevelIndex) = 0;
 
diff --git a/Engine/Source/Runtime/Engine/Private/HLOD/HLODEngineSubsystem.cpp b/Engine/Source/Runtime/Engine/Private/HLOD/HLODEngineSubsystem.cpp
--- a/Engine/Source/Runtime/Engine/Private/HLOD/HLODEngineSubsystem.cpp
+++ b/Engine/Source/Runtime/Engine/Private/HLOD/HLODEngineSubsystem.cpp
@@ -64,6 +64,7 @@ void UHLODEngineSubsystem::RegisterRecreateLODActorsDelegates()
 		OnPostWorldInitializationDelegateHandle = FWorldDelegates::OnPostWorldInitialization.AddUObject(this, &UHLODEngineSubsystem::RecreateLODActorsForWorld);
 		OnLevelAddedToWorldDelegateHandle = FWorldDelegates::LevelAddedToWorld.AddUObject(this, &UHLODEngineSubsystem::RecreateLODActorsForLevel);
 		OnPreSaveWorlDelegateHandle = FEditorDelegates::PreSaveWorldWithContext.AddUObject(this, &UHLODEngineSubsystem::OnPreSaveWorld);
+		OnWorldCollectSaveReferencesDelegateHandle = FWorldDelegates::OnCollectSaveReferences.AddUObject(this, &UHLODEngineSubsystem::OnWorldCollectSaveReferences);
 	}	
 }
 
@@ -194,5 +195,31 @@ void UHLODEngineSubsystem::OnPreSaveWorld(UWorld* InWorld, FObjectPreSaveContext
 	}
 }
 
+void UHLODEngineSubsystem::OnWorldCollectSaveReferences(UWorld* World, FArchive& Ar)
+{
+	// Record a dependency on the paths to HLODProxy packages that would be associated with
+	// this level if they exist
+	FHierarchicalLODUtilitiesModule& Module = FModuleManager::LoadModuleChecked<FHierarchicalLODUtilitiesModule>("HierarchicalLODUtilities");
+	IHierarchicalLODUtilities* Utilities = Module.GetUtilities();
+
+	for (ULevel* Level : World->GetLevels())
+	{
+		AWorldSettings* WorldSettings = Level->GetWorldSettings(false);
+		if (WorldSettings)
+		{
+			int32 NumLODLevels = WorldSettings->GetHierarchicalLODSetup().Num();
+			for (int32 LODIndex = 0; LODIndex < NumLODLevels; ++LODIndex)
+			{
+				const FString HLODLevelPackageName = Utilities->GetHLODPackageName(Level, LODIndex);
+				if (!HLODLevelPackageName.IsEmpty())
+				{
+					FSoftObjectPath Path(HLODLevelPackageName);
+					Ar << Path;
+				}
+			}
+		}
+	}
+}
+
 #endif // WITH_EDITOR
 
diff --git a/Engine/Source/Runtime/Engine/Public/HLOD/HLODEngineSubsystem.h b/Engine/Source/Runtime/Engine/Public/HLOD/HLODEngineSubsystem.h
--- a/Engine/Source/Runtime/Engine/Public/HLOD/HLODEngineSubsystem.h
+++ b/Engine/Source/Runtime/Engine/Public/HLOD/HLODEngineSubsystem.h
@@ -54,6 +54,7 @@ private:
 	void RecreateLODActorsForLevel(ULevel* InLevel, UWorld* InWorld);
 
 	void OnPreSaveWorld(UWorld* InWorld, FObjectPreSaveContext ObjectSaveContext);
+	void OnWorldCollectSaveReferences(UWorld* World, FArchive& Ar);
 
 	void UnregisterRecreateLODActorsDelegates();
 	void RegisterRecreateLODActorsDelegates();
@@ -64,6 +65,7 @@ private:
 	FDelegateHandle OnPostWorldInitializationDelegateHandle;
 	FDelegateHandle OnLevelAddedToWorldDelegateHandle;
 	FDelegateHandle OnPreSaveWorlDelegateHandle;
+	FDelegateHandle OnWorldCollectSaveReferencesDelegateHandle;
 
 	bool bDisableHLODCleanupOnLoad;
 	bool bDisableHLODSpawningOnLoad;
