Allow UPrimitiveComponent::LODParentPrimitive to be used outside of ALODActor
* Replaced CastCheckedby a null check in FHierarchicalLODUtilities::GetParentLODActor()

diff --git a/Engine/Source/Developer/HierarchicalLODUtilities/Private/HierarchicalLODUtilities.cpp b/Engine/Source/Developer/HierarchicalLODUtilities/Private/HierarchicalLODUtilities.cpp
--- a/Engine/Source/Developer/HierarchicalLODUtilities/Private/HierarchicalLODUtilities.cpp
+++ b/Engine/Source/Developer/HierarchicalLODUtilities/Private/HierarchicalLODUtilities.cpp
@@ -959,7 +959,6 @@ EClusterGenerationError FHierarchicalLODUtilities::ShouldGenerateCluster(AActor*
 
 ALODActor* FHierarchicalLODUtilities::GetParentLODActor(const AActor* InActor)
 {
-	ALODActor* ParentActor = nullptr;
 	if (InActor)
 	{
 		TArray<UStaticMeshComponent*> ComponentArray;
@@ -969,13 +968,15 @@ ALODActor* FHierarchicalLODUtilities::GetParentLODActor(const AActor* InActor)
 			UPrimitiveComponent* ParentComponent = Component->GetLODParentPrimitive();
 			if (ParentComponent)
 			{
-				ParentActor = CastChecked<ALODActor>(ParentComponent->GetOwner());
-				break;
+				if (ALODActor* ParentActor = Cast<ALODActor>(ParentComponent->GetOwner()))
+				{
+					return ParentActor;
+				}
 			}
 		}
 	}
 
-	return ParentActor;
+	return nullptr;
 }
 
 void FHierarchicalLODUtilities::DestroyCluster(ALODActor* InActor)
