HLOD - Modified FHierarchicalLODUtilities::ComputeStaticMeshLODLevel() so that it can work with cooked static meshes
* SourceModels are stripped out in cooked builds, but we can fallback on the RenderData

diff --git a/Engine/Source/Developer/HierarchicalLODUtilities/Private/HierarchicalLODUtilities.cpp b/Engine/Source/Developer/HierarchicalLODUtilities/Private/HierarchicalLODUtilities.cpp
--- a/Engine/Source/Developer/HierarchicalLODUtilities/Private/HierarchicalLODUtilities.cpp
+++ b/Engine/Source/Developer/HierarchicalLODUtilities/Private/HierarchicalLODUtilities.cpp
@@ -1277,13 +1277,16 @@ void FHierarchicalLODUtilities::DeleteLODActorsInHLODLevel(UWorld* InWorld, cons
 }
 
 int32 FHierarchicalLODUtilities::ComputeStaticMeshLODLevel(const TArray<FStaticMeshSourceModel>& SourceModels, const FStaticMeshRenderData* RenderData, const float ScreenSize)
-{	
-	const int32 NumLODs = SourceModels.Num();
+{
+	const int32 NumLODs = !SourceModels.IsEmpty() ? SourceModels.Num() : RenderData->LODResources.Num();
+
 	// Walk backwards and return the first matching LOD
 	for (int32 LODIndex = NumLODs - 1; LODIndex >= 0; --LODIndex)
 	{
+		const FPerPlatformFloat& LODSceenSize = SourceModels.IsValidIndex(LODIndex) && SourceModels[LODIndex].ScreenSize.Default != 0.0f ? SourceModels[LODIndex].ScreenSize : RenderData->ScreenSize[LODIndex];
+
 		// For HLOD generation we want the default values and not the per platform overrides
-		if (SourceModels[LODIndex].ScreenSize.Default > ScreenSize || ((SourceModels[LODIndex].ScreenSize.Default == 0.0f) && (RenderData->ScreenSize[LODIndex].Default != SourceModels[LODIndex].ScreenSize.Default) && (RenderData->ScreenSize[LODIndex].Default > ScreenSize)))
+		if (LODSceenSize.Default > ScreenSize)
 		{
 			return FMath::Max(LODIndex, 0);
 		}
@@ -1294,13 +1297,15 @@ int32 FHierarchicalLODUtilities::ComputeStaticMeshLODLevel(const TArray<FStaticM
 
 int32 FHierarchicalLODUtilities::GetLODLevelForScreenSize(const UStaticMeshComponent* StaticMeshComponent, const float ScreenSize)
 {
-	check(StaticMeshComponent != nullptr && StaticMeshComponent->GetStaticMesh() != nullptr);
+	check(StaticMeshComponent != nullptr);
+
+	const UStaticMesh* StaticMesh = StaticMeshComponent->GetStaticMesh();
+	check(StaticMesh != nullptr);
 
-	const FStaticMeshRenderData* RenderData = StaticMeshComponent->GetStaticMesh()->GetRenderData();
-	checkf(RenderData != nullptr, TEXT("StaticMesh in StaticMeshComponent %s contains invalid render data"), *StaticMeshComponent->GetName());
-	checkf(StaticMeshComponent->GetStaticMesh()->GetNumSourceModels() > 0, TEXT("StaticMesh in StaticMeshComponent %s contains no SourceModels"), *StaticMeshComponent->GetName());
+	const FStaticMeshRenderData* RenderData = StaticMesh->GetRenderData();
+	checkf(RenderData != nullptr, TEXT("StaticMesh %s referenced by StaticMeshComponent %s contains invalid render data"), *StaticMesh->GetName(), *StaticMeshComponent->GetName());
 
-	return ComputeStaticMeshLODLevel(const_cast<const UStaticMesh*>(ToRawPtr(StaticMeshComponent->GetStaticMesh()))->GetSourceModels(), RenderData, ScreenSize);
+	return ComputeStaticMeshLODLevel(StaticMesh->GetSourceModels(), RenderData, ScreenSize);
 }
 
 AHierarchicalLODVolume* FHierarchicalLODUtilities::CreateVolumeForLODActor(ALODActor* InLODActor, UWorld* InWorld)
