Enable bRayTracingFarField on meshes generated by legacy HLOD

#Synced-CL 49782167

diff --git a/Engine/Source/Developer/HierarchicalLODUtilities/Private/HierarchicalLODUtilities.cpp b/Engine/Source/Developer/HierarchicalLODUtilities/Private/HierarchicalLODUtilities.cpp
--- a/Engine/Source/Developer/HierarchicalLODUtilities/Private/HierarchicalLODUtilities.cpp
+++ b/Engine/Source/Developer/HierarchicalLODUtilities/Private/HierarchicalLODUtilities.cpp
@@ -768,7 +768,10 @@ bool FHierarchicalLODUtilities::BuildStaticMeshForLODActor(ALODActor* LODActor,
 
 			bDirtyPackage |= (LODActor->GetDrawDistance() != DrawDistance);
 			LODActor->SetDrawDistance(DrawDistance);
-			
+
+			bDirtyPackage |= (LODActor->GetStaticMeshComponent()->bRayTracingFarField != true);
+			LODActor->GetStaticMeshComponent()->bRayTracingFarField = true;
+
 			LODActor->DetermineShadowingFlags();
 
 			// Link proxy to actor
